Create slate UI for multi user server

diff --git a/Engine/Plugins/Developer/Concert/ConcertMain/Source/Concert/Public/ConcertSettings.h b/Engine/Plugins/Developer/Concert/ConcertMain/Source/Concert/Public/ConcertSettings.h
--- a/Engine/Plugins/Developer/Concert/ConcertMain/Source/Concert/Public/ConcertSettings.h
+++ b/Engine/Plugins/Developer/Concert/ConcertMain/Source/Concert/Public/ConcertSettings.h
@@ -331,6 +331,10 @@ public:
 	UPROPERTY(config, EditAnywhere, Category="Client Settings")
 	bool bRetryAutoConnectOnError = false;
 
+	/** If true, MU servers will be launched with Slate UI. This appends -WITHSLATE to the cmd.  */
+	UPROPERTY(config, EditAnywhere, Category="Client Settings")
+	bool bRunWithSlate = false;
+
 	/** 
 	 * Default server url (just a name for now) to look for on auto or default connect. 
  	 * Can be specified on the editor cmd with `-CONCERTSERVER=`.
diff --git a/Engine/Plugins/Developer/Concert/ConcertSync/ConcertSyncServer/Source/ConcertSyncServer/Public/ConcertSyncServerLoop.h b/Engine/Plugins/Developer/Concert/ConcertSync/ConcertSyncServer/Source/ConcertSyncServer/Public/ConcertSyncServerLoop.h
--- a/Engine/Plugins/Developer/Concert/ConcertSync/ConcertSyncServer/Source/ConcertSyncServer/Public/ConcertSyncServerLoop.h
+++ b/Engine/Plugins/Developer/Concert/ConcertSync/ConcertSyncServer/Source/ConcertSyncServer/Public/ConcertSyncServerLoop.h
@@ -3,38 +3,7 @@
 #pragma once
 
 #include "CoreMinimal.h"
-#include "ConcertMessageData.h"
-#include "ConcertSyncSessionFlags.h"
-
-class UConcertServerConfig;
-class IConcertFileSharingService;
-
-struct FConcertSyncServerLoopInitArgs
-{
-	/** Framerate that the main loop should try to Tick at */
-	int32 IdealFramerate = 60;
-
-	/** Flags controlling what features are enabled for sessions within this server */
-	EConcertSyncSessionFlags SessionFlags = EConcertSyncSessionFlags::None;
-
-	/** The role that this server will perform (eg, MultiUser, DisasterRecovery, etc) */
-	FString ServiceRole;
-
-	/** Friendly name to use for this service (when showing it to a user in log messages, etc) */
-	FString ServiceFriendlyName;
-
-	/** The session filter to apply when auto-archiving sessions on this server */
-	FConcertSessionFilter ServiceAutoArchiveSessionFilter;
-
-	/** The optional file sharing server to exchange large files. Can be null. */
-	TSharedPtr<IConcertFileSharingService> FileSharingService;
-
-	/** Function to get the server settings object to configure the server with with, or unbound to parse the default settings */
-	TFunction<const UConcertServerConfig*()> GetServerConfigFunc;
-
-	/** Whether the service should show the log console. */
-	bool bShowConsole = true;
-};
+#include "ConcertSyncServerLoopInitArgs.h"
 
 /**
  * Blocking main loop for running a Concert Sync server application.
diff --git a/Engine/Plugins/Developer/Concert/ConcertSync/ConcertSyncServer/Source/ConcertSyncServer/Public/ConcertSyncServerLoop.inl b/Engine/Plugins/Developer/Concert/ConcertSync/ConcertSyncServer/Source/ConcertSyncServer/Public/ConcertSyncServerLoop.inl
--- a/Engine/Plugins/Developer/Concert/ConcertSync/ConcertSyncServer/Source/ConcertSyncServer/Public/ConcertSyncServerLoop.inl
+++ b/Engine/Plugins/Developer/Concert/ConcertSync/ConcertSyncServer/Source/ConcertSyncServer/Public/ConcertSyncServerLoop.inl
@@ -90,6 +90,9 @@ int32 ConcertSyncServerLoop(int32 ArgC, TCHAR** ArgV, const FConcertSyncServerLo
 
 	if (Result >= 0)
 	{
+		// Give external modules to do early initialisation
+		InitArgs.PreInitServerLoop.Broadcast();
+		
 		TSharedPtr<IPlugin> Plugin = IPluginManager::Get().FindPlugin(TEXT("UdpMessaging"));
 		if (!Plugin || !Plugin->IsEnabled())
 		{
@@ -143,9 +146,10 @@ int32 ConcertSyncServerLoop(int32 ArgC, TCHAR** ArgV, const FConcertSyncServerLo
 				}
 			}
 		}
-
+		
 		UE_LOG(LogSyncServer, Display, TEXT("%s Initialized (Name: %s, Version: %d.%d, Role: %s)"), *InitArgs.ServiceFriendlyName, *ConcertSyncServer->GetConcertServer()->GetServerInfo().ServerName, ENGINE_MAJOR_VERSION, ENGINE_MINOR_VERSION, *ConcertSyncServer->GetConcertServer()->GetRole());
-
+		InitArgs.PostInitServerLoop.Broadcast(ConcertSyncServer.ToSharedRef());
+		
 		double LastTime = FPlatformTime::Seconds();
 		const float IdealFrameTime = 1.0f / InitArgs.IdealFramerate;
 
@@ -157,6 +161,7 @@ int32 ConcertSyncServerLoop(int32 ArgC, TCHAR** ArgV, const FConcertSyncServerLo
 			FTaskGraphInterface::Get().ProcessThreadUntilIdle(ENamedThreads::GameThread);
 
 			// Pump & Tick objects
+			InitArgs.TickPostGameThread.Broadcast(DeltaTime);
 			FTSTicker::GetCoreTicker().Tick(DeltaTime);
 
 			GFrameCounter++;
diff --git a/Engine/Plugins/Developer/Concert/ConcertSync/ConcertSyncServer/Source/ConcertSyncServer/Public/ConcertSyncServerLoopInitArgs.h b/Engine/Plugins/Developer/Concert/ConcertSync/ConcertSyncServer/Source/ConcertSyncServer/Public/ConcertSyncServerLoopInitArgs.h
new file mode 100644
--- /dev/null
+++ b/Engine/Plugins/Developer/Concert/ConcertSync/ConcertSyncServer/Source/ConcertSyncServer/Public/ConcertSyncServerLoopInitArgs.h
@@ -0,0 +1,58 @@
+// Copyright Epic Games, Inc. All Rights Reserved.
+
+#pragma once
+
+#include "CoreMinimal.h"
+#include "ConcertMessageData.h"
+#include "ConcertSyncSessionFlags.h"
+
+class UConcertServerConfig;
+class IConcertFileSharingService;
+class IConcertSyncServer;
+
+struct FConcertSyncServerLoopInitArgs
+{
+	DECLARE_MULTICAST_DELEGATE(FMultiUserServerLoopInitEvent);
+	DECLARE_MULTICAST_DELEGATE_OneParam(FMultiUserServerInitServerEvent, TSharedRef<IConcertSyncServer>);
+	DECLARE_MULTICAST_DELEGATE_OneParam(FMultiUserServerLoopTickEvent, double /* DeltaTime */);
+	
+	/** Framerate that the main loop should try to Tick at */
+	int32 IdealFramerate = 60;
+
+	/** Flags controlling what features are enabled for sessions within this server */
+	EConcertSyncSessionFlags SessionFlags = EConcertSyncSessionFlags::None;
+
+	/** The role that this server will perform (eg, MultiUser, DisasterRecovery, etc) */
+	FString ServiceRole;
+
+	/** Friendly name to use for this service (when showing it to a user in log messages, etc) */
+	FString ServiceFriendlyName;
+
+	/** The session filter to apply when auto-archiving sessions on this server */
+	FConcertSessionFilter ServiceAutoArchiveSessionFilter;
+
+	/** The optional file sharing server to exchange large files. Can be null. */
+	TSharedPtr<IConcertFileSharingService> FileSharingService;
+
+	/** Function to get the server settings object to configure the server with with, or unbound to parse the default settings */
+	TFunction<const UConcertServerConfig*()> GetServerConfigFunc;
+
+	/** Called before Multi User Server is initialized: no modules have yet been loaded. */
+	FMultiUserServerLoopInitEvent PreInitServerLoop;
+	
+	/**
+	 * Called while initilializing Multi User server loop. External services can perform custom initialization here.
+	 *
+	 * Called just before entering the server loop. LoadModulesForEnabledPlugins will have been called with PreDefault 
+	 * and Default and the multi user session created.
+	 */
+	FMultiUserServerInitServerEvent PostInitServerLoop;
+	
+	/**
+	 * Called after the game thread's tick and before FTSTicker::GetCoreTicker.
+	 */
+	FMultiUserServerLoopTickEvent TickPostGameThread;
+
+	/** Whether the service should show the log console. */
+	bool bShowConsole = true;
+};
diff --git a/Engine/Plugins/Developer/Concert/ConcertUI/ConcertServerUI/ConcertServerUI.uplugin b/Engine/Plugins/Developer/Concert/ConcertUI/ConcertServerUI/ConcertServerUI.uplugin
new file mode 100644
--- /dev/null
+++ b/Engine/Plugins/Developer/Concert/ConcertUI/ConcertServerUI/ConcertServerUI.uplugin
@@ -0,0 +1,41 @@
+{
+	"FileVersion": 3,
+	"Version": 1,
+	"VersionName": "0.2",
+	"FriendlyName": "Concert Sync Server UI",
+	"Description": "UI plugin that visualizes the multi-user server",
+	"Category": "Networking",
+	"CreatedBy": "Epic Games, Inc.",
+	"CreatedByURL": "http://epicgames.com",
+	"DocsURL": "",
+	"MarketplaceURL": "",
+	"SupportURL": "",
+	"EnabledByDefault": false,
+	"CanContainContent": false,
+	"Hidden": true,
+	"IsBetaVersion": true,
+	"Installed": false,
+	"SupportedPrograms": [ "UnrealMultiUserServer", "UnrealRecoverySvc", "CrashReportClientEditor" ],
+	"Modules": [
+        {
+          "Name": "ConcertServerUI",
+          "Type": "UncookedOnly",
+          "LoadingPhase": "Default",
+          "ProgramAllowList": [ "UnrealMultiUserServer", "UnrealRecoverySvc", "CrashReportClientEditor" ]
+        }
+	],
+	"Plugins": [
+		{
+			"Name": "ConcertMain",
+			"Enabled": true
+		},
+		{
+			"Name": "ConcertSyncCore",
+			"Enabled": true
+		},
+		{
+			"Name": "ConcertSyncServer",
+			"Enabled": true
+		}
+	]
+}
diff --git a/Engine/Plugins/Developer/Concert/ConcertUI/ConcertServerUI/Source/ConcertServerUI/ConcertServerUI.Build.cs b/Engine/Plugins/Developer/Concert/ConcertUI/ConcertServerUI/Source/ConcertServerUI/ConcertServerUI.Build.cs
new file mode 100644
--- /dev/null
+++ b/Engine/Plugins/Developer/Concert/ConcertUI/ConcertServerUI/Source/ConcertServerUI/ConcertServerUI.Build.cs
@@ -0,0 +1,46 @@
+// Copyright Epic Games, Inc. All Rights Reserved.
+
+namespace UnrealBuildTool.Rules
+{
+	public class ConcertServerUI : ModuleRules
+	{
+		public ConcertServerUI(ReadOnlyTargetRules Target) : base(Target)
+		{
+			// Make the linking path shorter (to prevent busting the Windows limit) when linking ConcertServerUI.lib against an executable that have a long name.
+			ShortName = "CncrtSvrUI";
+
+			PrivateDependencyModuleNames.AddRange(
+				new string[]
+				{
+					"ApplicationCore",
+					"Core",
+					"CoreUObject",
+					"Concert",
+					"ConcertSyncCore",
+					"ConcertTransport",
+					"EditorStyle",			// Needed so OutputLog module to work correctly
+					"InputCore",
+					"OutputLog",
+					"Projects",
+					"Slate",
+					"SlateCore",
+					"StandaloneRenderer",	
+					"ToolWidgets",
+					"WorkspaceMenuStructure" // Needed so OutputLog module to work correctly
+				}
+			);
+			
+			PrivateIncludePathModuleNames.AddRange(
+				new string[] {
+					"ConcertSyncServer",
+				}
+			);
+			
+			DynamicallyLoadedModuleNames.AddRange(
+				new string[]
+				{
+					"ConcertSyncServer"
+				});
+		}
+	}
+}
diff --git a/Engine/Plugins/Developer/Concert/ConcertUI/ConcertServerUI/Source/ConcertServerUI/Private/ConcertConsoleCommandExecutor.cpp b/Engine/Plugins/Developer/Concert/ConcertUI/ConcertServerUI/Source/ConcertServerUI/Private/ConcertConsoleCommandExecutor.cpp
new file mode 100644
--- /dev/null
+++ b/Engine/Plugins/Developer/Concert/ConcertUI/ConcertServerUI/Source/ConcertServerUI/Private/ConcertConsoleCommandExecutor.cpp
@@ -0,0 +1,88 @@
+﻿// Copyright Epic Games, Inc. All Rights Reserved.
+
+#include "ConcertConsoleCommandExecutor.h"
+
+#include "CoreGlobals.h"
+#include "Framework/Commands/InputChord.h"
+#include "Misc/OutputDevice.h"
+
+#define LOCTEXT_NAMESPACE "UnrealMultiUserUI"
+
+FName FConcertConsoleCommandExecutor::StaticName()
+{
+	static const FName CmdExecName = TEXT("Cmd");
+	return CmdExecName;
+}
+
+FName FConcertConsoleCommandExecutor::GetName() const
+{
+	return StaticName();
+}
+
+FText FConcertConsoleCommandExecutor::GetDisplayName() const
+{
+	return LOCTEXT("ConsoleCommandExecutorDisplayName", "Cmd");
+}
+
+FText FConcertConsoleCommandExecutor::GetDescription() const
+{
+	return LOCTEXT("ConsoleCommandExecutorDescription", "Execute Unreal Console Commands");
+}
+
+FText FConcertConsoleCommandExecutor::GetHintText() const
+{
+	return LOCTEXT("ConsoleCommandExecutorHintText", "Enter Console Command");
+}
+
+void FConcertConsoleCommandExecutor::GetAutoCompleteSuggestions(const TCHAR* Input, TArray<FString>& Out)
+{
+	auto OnConsoleVariable = [&Out](const TCHAR *Name, IConsoleObject* CVar)
+	{
+		if (CVar->TestFlags(ECVF_Unregistered))
+		{
+			return;
+		}
+
+		Out.Add(Name);
+	};
+
+	IConsoleManager::Get().ForEachConsoleObjectThatContains(FConsoleObjectVisitor::CreateLambda(OnConsoleVariable), Input);
+}
+
+void FConcertConsoleCommandExecutor::GetExecHistory(TArray<FString>& Out)
+{
+	IConsoleManager::Get().GetConsoleHistory(TEXT(""), Out);
+}
+
+bool FConcertConsoleCommandExecutor::Exec(const TCHAR* Input)
+{
+	IConsoleManager::Get().AddConsoleHistoryEntry(TEXT(""), Input);
+	
+	FOutputDevice& Output = *GLog;
+	UWorld* World = nullptr;
+	
+	bool bSuccess = IConsoleManager::Get().ProcessUserConsoleInput(Input, Output, World);
+	if (!bSuccess)
+	{
+		bSuccess |= FSelfRegisteringExec::StaticExec(World, Input, Output);
+	}
+
+	return bSuccess;
+}
+
+bool FConcertConsoleCommandExecutor::AllowHotKeyClose() const
+{
+	return true;
+}
+
+bool FConcertConsoleCommandExecutor::AllowMultiLine() const
+{
+	return false;
+}
+
+FInputChord FConcertConsoleCommandExecutor::GetHotKey() const
+{
+	return FInputChord();
+}
+
+#undef LOCTEXT_NAMESPACE
diff --git a/Engine/Plugins/Developer/Concert/ConcertUI/ConcertServerUI/Source/ConcertServerUI/Private/ConcertConsoleCommandExecutor.h b/Engine/Plugins/Developer/Concert/ConcertUI/ConcertServerUI/Source/ConcertServerUI/Private/ConcertConsoleCommandExecutor.h
new file mode 100644
--- /dev/null
+++ b/Engine/Plugins/Developer/Concert/ConcertUI/ConcertServerUI/Source/ConcertServerUI/Private/ConcertConsoleCommandExecutor.h
@@ -0,0 +1,27 @@
+﻿// Copyright Epic Games, Inc. All Rights Reserved.
+
+#pragma once
+
+#include "CoreMinimal.h"
+#include "HAL/IConsoleManager.h"
+
+/** Executes commands put into the output log */
+class FConcertConsoleCommandExecutor : public IConsoleCommandExecutor
+{
+public:
+	
+	static FName StaticName();
+
+	//~ Begin IConsoleCommandExecutor Interface
+	virtual FName GetName() const override;
+	virtual FText GetDisplayName() const override;
+	virtual FText GetDescription() const override;
+	virtual FText GetHintText() const override;
+	virtual void GetAutoCompleteSuggestions(const TCHAR* Input, TArray<FString>& Out) override;
+	virtual void GetExecHistory(TArray<FString>& Out) override;
+	virtual bool Exec(const TCHAR* Input) override;
+	virtual bool AllowHotKeyClose() const override;
+	virtual bool AllowMultiLine() const override;
+	virtual FInputChord GetHotKey() const override;
+	//~ End IConsoleCommandExecutor Interface
+};
\ No newline at end of file
diff --git a/Engine/Plugins/Developer/Concert/ConcertUI/ConcertServerUI/Source/ConcertServerUI/Private/ConcertServerStyle.cpp b/Engine/Plugins/Developer/Concert/ConcertUI/ConcertServerUI/Source/ConcertServerUI/Private/ConcertServerStyle.cpp
new file mode 100644
--- /dev/null
+++ b/Engine/Plugins/Developer/Concert/ConcertUI/ConcertServerUI/Source/ConcertServerUI/Private/ConcertServerStyle.cpp
@@ -0,0 +1,77 @@
+// Copyright Epic Games, Inc. All Rights Reserved.
+
+#include "ConcertServerStyle.h"
+
+#include "EditorStyleSet.h"
+#include "Framework/Application/SlateApplication.h"
+#include "Interfaces/IPluginManager.h"
+#include "Styling/SlateStyleRegistry.h"
+
+TSharedPtr<FSlateStyleSet> FConcertServerStyle::StyleInstance = nullptr;
+
+void FConcertServerStyle::Initialize()
+{
+	// The core style must be initialized before the editor style
+	//FSlateApplication::InitializeCoreStyle();
+	//FEditorStyle::ResetToDefault();
+
+	if (!StyleInstance.IsValid())
+	{
+		StyleInstance = Create();
+		FSlateStyleRegistry::RegisterSlateStyle(*StyleInstance);
+	}
+}
+
+void FConcertServerStyle::Shutdown()
+{
+	if (StyleInstance)
+	{
+		
+		FSlateStyleRegistry::UnRegisterSlateStyle(*StyleInstance);
+		ensure(StyleInstance.IsUnique());
+		StyleInstance.Reset();
+	}
+}
+
+const ISlateStyle& FConcertServerStyle::Get()
+{
+	return *StyleInstance;
+}
+
+FName FConcertServerStyle::GetStyleSetName()
+{
+	static FName StyleSetName(TEXT("ConcertServerStyle"));
+	return StyleSetName;
+}
+
+#define IMAGE_PLUGIN_BRUSH( RelativePath, ... ) FSlateImageBrush( FConcertServerStyle::InContent( RelativePath, ".png" ), __VA_ARGS__ )
+#define IMAGE_BRUSH(RelativePath, ...) FSlateImageBrush(StyleSet->RootToContentDir(RelativePath, TEXT(".png")), __VA_ARGS__)
+#define IMAGE_PLUGIN_BRUSH_SVG( RelativePath, ... ) FSlateVectorImageBrush( FConcertServerStyle::InContent(RelativePath, ".svg"), __VA_ARGS__)
+
+FString FConcertServerStyle::InContent(const FString& RelativePath, const ANSICHAR* Extension)
+{
+	static FString ContentDir = IPluginManager::Get().FindPlugin(TEXT("ConcertServerUI"))->GetContentDir();
+	return (ContentDir / RelativePath) + Extension;
+}
+
+TSharedRef<FSlateStyleSet> FConcertServerStyle::Create()
+{
+	TSharedRef<FSlateStyleSet> StyleSet = MakeShared<FSlateStyleSet>(GetStyleSetName());
+	
+	StyleSet->SetContentRoot(FPaths::EngineContentDir() / TEXT("Slate/Starship/Insights"));
+	StyleSet->SetCoreContentRoot(FPaths::EngineContentDir() / TEXT("Slate"));
+	
+	const FVector2D Icon12x12(12.0f, 12.0f); 
+	const FVector2D Icon16x16(16.0f, 16.0f); 
+	const FVector2D Icon20x20(20.0f, 20.0f); 
+	const FVector2D Icon32x32(20.0f, 20.0f); 
+	
+	StyleSet->Set("Concert.MultiUser", new IMAGE_PLUGIN_BRUSH("Icons/icon_MultiUser_32x", Icon32x32));
+	return StyleSet;
+}
+
+#undef TODO_IMAGE_BRUSH
+#undef EDITOR_BOX_BRUSH
+#undef EDITOR_IMAGE_BRUSH_SVG
+#undef EDITOR_IMAGE_BRUSH
+
diff --git a/Engine/Plugins/Developer/Concert/ConcertUI/ConcertServerUI/Source/ConcertServerUI/Private/ConcertServerStyle.h b/Engine/Plugins/Developer/Concert/ConcertUI/ConcertServerUI/Source/ConcertServerUI/Private/ConcertServerStyle.h
new file mode 100644
--- /dev/null
+++ b/Engine/Plugins/Developer/Concert/ConcertUI/ConcertServerUI/Source/ConcertServerUI/Private/ConcertServerStyle.h
@@ -0,0 +1,25 @@
+// Copyright Epic Games, Inc. All Rights Reserved.
+
+#pragma once
+
+#include "CoreMinimal.h"
+#include "Styling/SlateStyle.h"
+
+/** Style data for Insights tools */
+class FConcertServerStyle
+{
+public:
+	static void Initialize();
+	static void Shutdown();
+
+	static const ISlateStyle& Get();
+	static FName GetStyleSetName();
+
+private:
+	
+	static TSharedPtr<FSlateStyleSet> StyleInstance;
+	
+	static FString InContent(const FString& RelativePath, const ANSICHAR* Extension);
+	static TSharedRef<FSlateStyleSet> Create();
+
+};
diff --git a/Engine/Plugins/Developer/Concert/ConcertUI/ConcertServerUI/Source/ConcertServerUI/Private/ConcertServerUIModule.cpp b/Engine/Plugins/Developer/Concert/ConcertUI/ConcertServerUI/Source/ConcertServerUI/Private/ConcertServerUIModule.cpp
new file mode 100644
--- /dev/null
+++ b/Engine/Plugins/Developer/Concert/ConcertUI/ConcertServerUI/Source/ConcertServerUI/Private/ConcertServerUIModule.cpp
@@ -0,0 +1,87 @@
+// Copyright Epic Games, Inc. All Rights Reserved.
+
+#include "IConcertServerUIModule.h"
+
+#include "ConcertConsoleCommandExecutor.h"
+#include "ConcertServerStyle.h"
+#include "ConcertSyncServerLoopInitArgs.h"
+#include "OutputLogModule.h"
+#include "Widgets/ConcertServerWindowController.h"
+
+#include "Framework/Application/SlateApplication.h"
+#include "StandaloneRenderer.h"
+#include "Features/IModularFeatures.h"
+
+#define LOCTEXT_NAMESPACE "UnrealMultiUserUI"
+
+/**
+ * 
+ */
+class FConcertServerUIModule : public IConcertServerUIModule
+{
+public:
+	virtual void StartupModule() override 
+	{
+		MultiUserServerLayoutIni = GConfig->GetConfigFilename(TEXT("MultiUserServerLayout"));
+	}
+	
+	virtual void ShutdownModule() override
+	{
+		FConcertServerStyle::Shutdown();
+		
+		WindowController.Reset();
+		FSlateApplication::Shutdown();
+	}
+	
+	virtual void InitSlateForServer(FConcertSyncServerLoopInitArgs& InitArgs) override
+	{
+		if (ensureMsgf(!WindowController.IsValid(), TEXT("InitSlateForServer is designed to be called at most once to create UI to run alongside the Multi User server.")))
+		{
+			PreInitializeMultiUser();
+			InitArgs.PostInitServerLoop.AddRaw(this, &FConcertServerUIModule::InitializeSlateApplication);
+			InitArgs.TickPostGameThread.AddRaw(this, &FConcertServerUIModule::TickSlate);
+		}
+	}
+	
+private:
+
+	/** Config path storing layout config. */
+	FString MultiUserServerLayoutIni;
+
+	/** Handles execution of commands */
+	TUniquePtr<FConcertConsoleCommandExecutor> CommandExecutor;
+	
+	/** Creates and manages window. Only one instance per application. */
+	TSharedPtr<FConcertServerWindowController> WindowController;
+
+	void PreInitializeMultiUser()
+	{
+		FModuleManager::Get().LoadModuleChecked("EditorStyle");
+		FConcertServerStyle::Initialize();
+
+		// Log history must be initialized before MU server loop init prints any messages	
+		FModuleManager::Get().LoadModuleChecked("OutputLog");
+	}
+	
+	void InitializeSlateApplication(TSharedRef<IConcertSyncServer> SyncServer)
+	{
+		FSlateApplication::InitializeAsStandaloneApplication(GetStandardStandaloneRenderer());
+		const FText ApplicationTitle = LOCTEXT("AppTitle", "Unreal Multi User Server");
+		FGlobalTabmanager::Get()->SetApplicationTitle(ApplicationTitle);
+
+		CommandExecutor = MakeUnique<FConcertConsoleCommandExecutor>();
+		IModularFeatures::Get().RegisterModularFeature(IConsoleCommandExecutor::ModularFeatureName(), CommandExecutor.Get());
+        
+        WindowController = MakeShared<FConcertServerWindowController>(FConcertServerWindowInitParams{ SyncServer, MultiUserServerLayoutIni });
+	}
+	
+	void TickSlate(double Tick)
+	{
+		FSlateApplication::Get().PumpMessages();
+		FSlateApplication::Get().Tick();
+	}
+};
+
+IMPLEMENT_MODULE(FConcertServerUIModule, ConcertSyncServerUI);
+
+#undef LOCTEXT_NAMESPACE
\ No newline at end of file
diff --git a/Engine/Plugins/Developer/Concert/ConcertUI/ConcertServerUI/Source/ConcertServerUI/Private/Widgets/Browser/ConcertServerSessionBrowserController.cpp b/Engine/Plugins/Developer/Concert/ConcertUI/ConcertServerUI/Source/ConcertServerUI/Private/Widgets/Browser/ConcertServerSessionBrowserController.cpp
new file mode 100644
--- /dev/null
+++ b/Engine/Plugins/Developer/Concert/ConcertUI/ConcertServerUI/Source/ConcertServerUI/Private/Widgets/Browser/ConcertServerSessionBrowserController.cpp
@@ -0,0 +1,35 @@
+﻿// Copyright Epic Games, Inc. All Rights Reserved.
+
+#include "ConcertServerSessionBrowserController.h"
+
+#include "ConcertServerStyle.h"
+#include "SConcertServerSessionBrowser.h"
+#include "Widgets/ConcertServerTabs.h"
+
+#define LOCTEXT_NAMESPACE "UnrealMultiUserUI"
+
+void FConcertServerSessionBrowserController::Init(const FConcertComponentInitParams& Params)
+{
+	FGlobalTabmanager::Get()->RegisterTabSpawner(
+			ConcertServerTabs::GetSessionBrowserTabId(),
+			FOnSpawnTab::CreateRaw(this, &FConcertServerSessionBrowserController::SpawnSessionBrowserTab)
+		)
+		.SetDisplayName(LOCTEXT("SessionBrowserTabTitle", "Session Browser"))
+		.SetTooltipText(LOCTEXT("SessionBrowserTooltipText", "A section to browse, start, archive, and restore server sessions."))
+		.SetIcon(FSlateIcon(FConcertServerStyle::GetStyleSetName(), TEXT("Concert.MultiUser")));
+}
+
+TSharedRef<SDockTab> FConcertServerSessionBrowserController::SpawnSessionBrowserTab(const FSpawnTabArgs& Args)
+{
+	const TSharedRef<SDockTab> DockTab = SNew(SDockTab)
+		.Label(LOCTEXT("SessionBrowserTabTitle", "Sessions"))
+		.TabRole(MajorTab)
+		[
+			SAssignNew(ConcertBrowser, SConcertServerSessionBrowser, SharedThis(this))
+		];
+
+	FGlobalTabmanager::Get()->SetMainTab(DockTab);
+	return DockTab;
+}
+
+#undef LOCTEXT_NAMESPACE 
diff --git a/Engine/Plugins/Developer/Concert/ConcertUI/ConcertServerUI/Source/ConcertServerUI/Private/Widgets/Browser/ConcertServerSessionBrowserController.h b/Engine/Plugins/Developer/Concert/ConcertUI/ConcertServerUI/Source/ConcertServerUI/Private/Widgets/Browser/ConcertServerSessionBrowserController.h
new file mode 100644
--- /dev/null
+++ b/Engine/Plugins/Developer/Concert/ConcertUI/ConcertServerUI/Source/ConcertServerUI/Private/Widgets/Browser/ConcertServerSessionBrowserController.h
@@ -0,0 +1,31 @@
+﻿// Copyright Epic Games, Inc. All Rights Reserved.
+
+#pragma once
+
+#include "CoreMinimal.h"
+#include "Widgets/Browser/IConcertServerSessionBrowserController.h"
+#include "Widgets/IConcertComponent.h"
+
+class FSpawnTabArgs;
+class SDockTab;
+class SConcertServerSessionBrowser;
+
+/** Implements controller in model-view-controller pattern. */
+class FConcertServerSessionBrowserController
+	:
+	public TSharedFromThis<FConcertServerSessionBrowserController>,
+	public IConcertComponent,
+	public IConcertServerSessionBrowserController
+{
+public:
+
+	//~ Begin IConcertComponent Interface
+	virtual void Init(const FConcertComponentInitParams& Params) override;
+	//~ End IConcertComponent Interface
+
+private:
+
+	TSharedPtr<SConcertServerSessionBrowser> ConcertBrowser;
+	
+	TSharedRef<SDockTab> SpawnSessionBrowserTab(const FSpawnTabArgs& Args);
+};
diff --git a/Engine/Plugins/Developer/Concert/ConcertUI/ConcertServerUI/Source/ConcertServerUI/Private/Widgets/Browser/IConcertServerSessionBrowserController.h b/Engine/Plugins/Developer/Concert/ConcertUI/ConcertServerUI/Source/ConcertServerUI/Private/Widgets/Browser/IConcertServerSessionBrowserController.h
new file mode 100644
--- /dev/null
+++ b/Engine/Plugins/Developer/Concert/ConcertUI/ConcertServerUI/Source/ConcertServerUI/Private/Widgets/Browser/IConcertServerSessionBrowserController.h
@@ -0,0 +1,14 @@
+﻿// Copyright Epic Games, Inc. All Rights Reserved.
+
+#pragma once
+
+#include "CoreMinimal.h"
+
+class IConcertServerSessionBrowserController
+{
+public:
+
+	
+	
+	virtual ~IConcertServerSessionBrowserController() = default;
+};
diff --git a/Engine/Plugins/Developer/Concert/ConcertUI/ConcertServerUI/Source/ConcertServerUI/Private/Widgets/Browser/SConcertServerSessionBrowser.cpp b/Engine/Plugins/Developer/Concert/ConcertUI/ConcertServerUI/Source/ConcertServerUI/Private/Widgets/Browser/SConcertServerSessionBrowser.cpp
new file mode 100644
--- /dev/null
+++ b/Engine/Plugins/Developer/Concert/ConcertUI/ConcertServerUI/Source/ConcertServerUI/Private/Widgets/Browser/SConcertServerSessionBrowser.cpp
@@ -0,0 +1,89 @@
+﻿// Copyright Epic Games, Inc. All Rights Reserved.
+
+#include "SConcertServerSessionBrowser.h"
+
+#include "SPositiveActionButton.h"
+#include "OutputLog/Public/OutputLogModule.h"
+#include "Widgets/Input/SSearchBox.h"
+#include "Widgets/Layout/SBorder.h"
+#include "Widgets/SBoxPanel.h"
+
+#define LOCTEXT_NAMESPACE "UnrealMultiUserUI"
+
+void SConcertServerSessionBrowser::Construct(const FArguments& InArgs, TWeakPtr<IConcertServerSessionBrowserController> InController)
+{
+	ChildSlot
+	[
+		SNew(SBorder)
+		.BorderImage(FAppStyle::Get().GetBrush("ToolPanel.GroupBorder"))
+		.Padding(FMargin(1.0f, 2.0f))
+		[
+			SNew(SVerticalBox)
+
+			+SVerticalBox::Slot()
+			.AutoHeight()
+			[
+				MakeControlBar()
+			]
+
+			// Session list.
+			+SVerticalBox::Slot()
+			.FillHeight(1.0f)
+			.Padding(1.0f, 2.0f)
+			[
+				MakeSessionTableView()
+			]
+		]
+	];
+}
+
+TSharedRef<SWidget> SConcertServerSessionBrowser::MakeControlBar()
+{
+	return SNew(SHorizontalBox)
+
+		// + New Session
+		+SHorizontalBox::Slot()
+		.AutoWidth()
+		[
+			SNew(SPositiveActionButton)
+				.OnClicked(this, &SConcertServerSessionBrowser::OnNewSessionClicked)
+				.Text(LOCTEXT("NewSession", "New Session"))
+		]
+
+		// The search text.
+		+SHorizontalBox::Slot()
+		.FillWidth(1.0)
+		.Padding(4.0f, 3.0f, 8.0f, 3.0f)
+		[
+			SAssignNew(SearchBox, SSearchBox)
+			.HintText(LOCTEXT("SearchHint", "Search sessions"))
+			.OnTextChanged(this, &SConcertServerSessionBrowser::OnSearchTextChanged)
+			.OnTextCommitted(this, &SConcertServerSessionBrowser::OnSearchTextCommitted)
+			.DelayChangeNotificationsWhileTyping(true)
+		]
+	;
+}
+
+TSharedRef<SWidget> SConcertServerSessionBrowser::MakeSessionTableView()
+{
+	//return FOutputLogModule::Get().MakeOutputLogDrawerWidget(FSimpleDelegate());
+	return SNullWidget::NullWidget;
+}
+
+FReply SConcertServerSessionBrowser::OnNewSessionClicked()
+{
+	// TODO:
+	return FReply::Handled();
+}
+
+void SConcertServerSessionBrowser::OnSearchTextChanged(const FText& InFilterText)
+{
+	// TODO:
+}
+
+void SConcertServerSessionBrowser::OnSearchTextCommitted(const FText& InFilterText, ETextCommit::Type CommitType)
+{
+	// TODO:
+}
+
+#undef LOCTEXT_NAMESPACE
diff --git a/Engine/Plugins/Developer/Concert/ConcertUI/ConcertServerUI/Source/ConcertServerUI/Private/Widgets/Browser/SConcertServerSessionBrowser.h b/Engine/Plugins/Developer/Concert/ConcertUI/ConcertServerUI/Source/ConcertServerUI/Private/Widgets/Browser/SConcertServerSessionBrowser.h
new file mode 100644
--- /dev/null
+++ b/Engine/Plugins/Developer/Concert/ConcertUI/ConcertServerUI/Source/ConcertServerUI/Private/Widgets/Browser/SConcertServerSessionBrowser.h
@@ -0,0 +1,35 @@
+﻿// Copyright Epic Games, Inc. All Rights Reserved.
+
+#pragma once
+
+#include "CoreMinimal.h"
+#include "Widgets/DeclarativeSyntaxSupport.h"
+#include "Widgets/SCompoundWidget.h"
+
+class IConcertServerSessionBrowserController;
+class SSearchBox;
+
+/** Shows a list of server sessions */
+class SConcertServerSessionBrowser : public SCompoundWidget
+{
+public:
+	SLATE_BEGIN_ARGS(SConcertServerSessionBrowser) { }
+	SLATE_END_ARGS()
+
+	void Construct(const FArguments& InArgs, TWeakPtr<IConcertServerSessionBrowserController> InController);
+
+private:
+
+	/** We can ask the controller about information and notify it about UI events. */
+	TWeakPtr<IConcertServerSessionBrowserController> Controller;
+	
+	TSharedPtr<SSearchBox> SearchBox;
+	
+	TSharedRef<SWidget> MakeControlBar();
+	TSharedRef<SWidget> MakeSessionTableView();
+
+	FReply OnNewSessionClicked();
+
+	void OnSearchTextChanged(const FText& InFilterText);
+	void OnSearchTextCommitted(const FText& InFilterText, ETextCommit::Type CommitType);
+};
diff --git a/Engine/Plugins/Developer/Concert/ConcertUI/ConcertServerUI/Source/ConcertServerUI/Private/Widgets/ConcertServerTabs.cpp b/Engine/Plugins/Developer/Concert/ConcertUI/ConcertServerUI/Source/ConcertServerUI/Private/Widgets/ConcertServerTabs.cpp
new file mode 100644
--- /dev/null
+++ b/Engine/Plugins/Developer/Concert/ConcertUI/ConcertServerUI/Source/ConcertServerUI/Private/Widgets/ConcertServerTabs.cpp
@@ -0,0 +1,15 @@
+﻿// Copyright Epic Games, Inc. All Rights Reserved.
+
+#include "ConcertServerTabs.h"
+
+FName ConcertServerTabs::GetSessionBrowserTabId()
+{
+	static FName SessionBrowserTabId("SessionBrowserTabId");
+	return SessionBrowserTabId;
+}
+
+FName ConcertServerTabs::GetOutputLogTabId()
+{
+	static FName OutputLogTabId("OutputLogTabId");
+	return OutputLogTabId;
+}
\ No newline at end of file
diff --git a/Engine/Plugins/Developer/Concert/ConcertUI/ConcertServerUI/Source/ConcertServerUI/Private/Widgets/ConcertServerTabs.h b/Engine/Plugins/Developer/Concert/ConcertUI/ConcertServerUI/Source/ConcertServerUI/Private/Widgets/ConcertServerTabs.h
new file mode 100644
--- /dev/null
+++ b/Engine/Plugins/Developer/Concert/ConcertUI/ConcertServerUI/Source/ConcertServerUI/Private/Widgets/ConcertServerTabs.h
@@ -0,0 +1,11 @@
+﻿// Copyright Epic Games, Inc. All Rights Reserved.
+
+#pragma once
+
+#include "CoreMinimal.h"
+
+namespace ConcertServerTabs
+{
+	static FName GetSessionBrowserTabId();
+	static FName GetOutputLogTabId();
+}
diff --git a/Engine/Plugins/Developer/Concert/ConcertUI/ConcertServerUI/Source/ConcertServerUI/Private/Widgets/ConcertServerWindowController.cpp b/Engine/Plugins/Developer/Concert/ConcertUI/ConcertServerUI/Source/ConcertServerUI/Private/Widgets/ConcertServerWindowController.cpp
new file mode 100644
--- /dev/null
+++ b/Engine/Plugins/Developer/Concert/ConcertUI/ConcertServerUI/Source/ConcertServerUI/Private/Widgets/ConcertServerWindowController.cpp
@@ -0,0 +1,90 @@
+// Copyright Epic Games, Inc. All Rights Reserved.
+
+#include "ConcertServerWindowController.h"
+
+#include "ConcertServerTabs.h"
+#include "Browser/ConcertServerSessionBrowserController.h"
+#include "OutputLog/OutputLogController.h"
+
+#include "Framework/Docking/TabManager.h"
+#include "Framework/Docking/LayoutService.h"
+#include "Framework/Notifications/NotificationManager.h"
+#include "HAL/PlatformApplicationMisc.h"
+#include "OutputLog/Public/OutputLogModule.h"
+#include "Widgets/SWindow.h"
+
+#define LOCTEXT_NAMESPACE "UnrealMultiUserUI"
+
+FConcertServerWindowController::FConcertServerWindowController(const FConcertServerWindowInitParams& Params)
+	: MultiUserServerLayoutIni(Params.MultiUserServerLayoutIni)
+{
+	SessionBrowserController = MakeShared<FConcertServerSessionBrowserController>();
+	OutputLogController = MakeShared<FOutputLogController>();
+	InitComponents(Params);
+
+	FDisplayMetrics DisplayMetrics;
+	FSlateApplication::Get().GetDisplayMetrics(DisplayMetrics);
+	const float DPIScaleFactor = FPlatformApplicationMisc::GetDPIScaleFactorAtPoint(DisplayMetrics.PrimaryDisplayWorkAreaRect.Left, DisplayMetrics.PrimaryDisplayWorkAreaRect.Top);
+
+	const bool bEmbedTitleAreaContent = false;
+	const FVector2D ClientSize(960.0f * DPIScaleFactor, 640.0f * DPIScaleFactor);
+	TSharedRef<SWindow> RootWindow = SNew(SWindow)
+		.Title(LOCTEXT("WindowTitle", "Unreal Multi User Server"))
+		.CreateTitleBar(!bEmbedTitleAreaContent)
+		.SupportsMaximize(true)
+		.SupportsMinimize(true)
+		.IsInitiallyMaximized(false)
+		.IsInitiallyMinimized(false)
+		.SizingRule(ESizingRule::UserSized)
+		.AutoCenter(EAutoCenter::PreferredWorkArea)
+		.ClientSize(ClientSize)
+		.AdjustInitialSizeAndPositionForDPIScale(false);
+		
+	const bool bShowRootWindowImmediately = false;
+	FSlateApplication::Get().AddWindow(RootWindow, bShowRootWindowImmediately);
+	FGlobalTabmanager::Get()->SetRootWindow(RootWindow);
+	FGlobalTabmanager::Get()->SetAllowWindowMenuBar(true);
+
+	FSlateNotificationManager::Get().SetRootWindow(RootWindow);
+	TSharedRef<FTabManager::FLayout> DefaultLayout = FTabManager::NewLayout("UnrealMultiUserServerLayout_v1.0");
+	DefaultLayout->AddArea
+	(
+		FTabManager::NewPrimaryArea()
+		->Split
+		(
+			FTabManager::NewStack()
+				->AddTab(ConcertServerTabs::GetSessionBrowserTabId(), ETabState::OpenedTab)
+				->AddTab(ConcertServerTabs::GetOutputLogTabId(), ETabState::OpenedTab)
+				->SetForegroundTab(ConcertServerTabs::GetSessionBrowserTabId())
+		)
+	);
+	
+	PersistentLayout = FLayoutSaveRestore::LoadFromConfig(Params.MultiUserServerLayoutIni, DefaultLayout);
+	TSharedPtr<SWidget> Content = FGlobalTabmanager::Get()->RestoreFrom(PersistentLayout.ToSharedRef(), RootWindow, bEmbedTitleAreaContent, EOutputCanBeNullptr::Never);
+	RootWindow->SetContent(Content.ToSharedRef());
+
+	RootWindow->SetOnWindowClosed(FOnWindowClosed::CreateRaw(this, &FConcertServerWindowController::OnWindowClosed));
+	RootWindow->ShowWindow();
+	const bool bForceWindowToFront = true;
+	RootWindow->BringToFront(bForceWindowToFront);
+}
+
+void FConcertServerWindowController::InitComponents(const FConcertServerWindowInitParams& WindowInitParams) const
+{
+	const FConcertComponentInitParams Params { WindowInitParams.Server };
+	SessionBrowserController->Init(Params);
+	OutputLogController->Init(Params);
+}
+
+void FConcertServerWindowController::OnWindowClosed(const TSharedRef<SWindow>& Window)
+{
+	SaveLayout();
+}
+
+void FConcertServerWindowController::SaveLayout() const
+{
+	FLayoutSaveRestore::SaveToConfig(MultiUserServerLayoutIni, PersistentLayout.ToSharedRef());
+    GConfig->Flush(false, MultiUserServerLayoutIni);
+}
+
+#undef LOCTEXT_NAMESPACE 
\ No newline at end of file
diff --git a/Engine/Plugins/Developer/Concert/ConcertUI/ConcertServerUI/Source/ConcertServerUI/Private/Widgets/ConcertServerWindowController.h b/Engine/Plugins/Developer/Concert/ConcertUI/ConcertServerUI/Source/ConcertServerUI/Private/Widgets/ConcertServerWindowController.h
new file mode 100644
--- /dev/null
+++ b/Engine/Plugins/Developer/Concert/ConcertUI/ConcertServerUI/Source/ConcertServerUI/Private/Widgets/ConcertServerWindowController.h
@@ -0,0 +1,49 @@
+// Copyright Epic Games, Inc. All Rights Reserved.
+
+#pragma once
+
+#include "CoreMinimal.h"
+
+class IConcertSyncServer;
+class FConcertServerSessionBrowserController;
+class FOutputLogController;
+
+struct FConcertServerWindowInitParams
+{
+	/** The server that the window is supposed to manage */
+	TSharedRef<IConcertSyncServer> Server;
+	
+	/** Config path for server layout ini */
+	FString MultiUserServerLayoutIni;
+	
+	FConcertServerWindowInitParams(TSharedRef<IConcertSyncServer> Server, FString MultiUserServerLayoutIni = FString())
+		: Server(Server)
+		, MultiUserServerLayoutIni(MoveTemp(MultiUserServerLayoutIni))
+	{}
+};
+
+/** Responsible for creating the Slate window for the server. Implements controller in the model-view-controller pattern. */
+class FConcertServerWindowController : public TSharedFromThis<FConcertServerWindowController>
+{
+public:
+
+	FConcertServerWindowController(const FConcertServerWindowInitParams& Params);
+
+private:
+
+	/** The ini file to use for saving the layout */
+	FString MultiUserServerLayoutIni;
+	
+	/** Holds the current layout for saving later. */
+	TSharedPtr<FTabManager::FLayout> PersistentLayout;
+
+	/** Manages the session browser */
+	TSharedPtr<FConcertServerSessionBrowserController> SessionBrowserController;
+	/** Manages tab for output log */
+	TSharedPtr<FOutputLogController> OutputLogController;
+	
+	void InitComponents(const FConcertServerWindowInitParams& WindowInitParams) const;
+
+	void OnWindowClosed(const TSharedRef<SWindow>& Window);
+	void SaveLayout() const;
+};
\ No newline at end of file
diff --git a/Engine/Plugins/Developer/Concert/ConcertUI/ConcertServerUI/Source/ConcertServerUI/Private/Widgets/IConcertComponent.h b/Engine/Plugins/Developer/Concert/ConcertUI/ConcertServerUI/Source/ConcertServerUI/Private/Widgets/IConcertComponent.h
new file mode 100644
--- /dev/null
+++ b/Engine/Plugins/Developer/Concert/ConcertUI/ConcertServerUI/Source/ConcertServerUI/Private/Widgets/IConcertComponent.h
@@ -0,0 +1,21 @@
+﻿// Copyright Epic Games, Inc. All Rights Reserved.
+
+#pragma once
+
+#include "CoreMinimal.h"
+
+struct FConcertComponentInitParams
+{
+	TSharedRef<IConcertSyncServer> Server;
+};
+
+/** Provides the base interface for elements in the concert server UI. */
+class IConcertComponent
+{
+public:
+
+	virtual ~IConcertComponent() = default;
+
+	/** Initialises the component, e.g. registering tab spawners, etc. */
+	virtual void Init(const FConcertComponentInitParams& Params) {}
+};
diff --git a/Engine/Plugins/Developer/Concert/ConcertUI/ConcertServerUI/Source/ConcertServerUI/Private/Widgets/OutputLog/OutputLogController.cpp b/Engine/Plugins/Developer/Concert/ConcertUI/ConcertServerUI/Source/ConcertServerUI/Private/Widgets/OutputLog/OutputLogController.cpp
new file mode 100644
--- /dev/null
+++ b/Engine/Plugins/Developer/Concert/ConcertUI/ConcertServerUI/Source/ConcertServerUI/Private/Widgets/OutputLog/OutputLogController.cpp
@@ -0,0 +1,45 @@
+﻿// Copyright Epic Games, Inc. All Rights Reserved.
+
+#include "OutputLogController.h"
+
+#include "Widgets/ConcertServerTabs.h"
+
+#include "EditorStyleSet.h"
+#include "Framework/Docking/TabManager.h"
+#include "OutputLogCreationParams.h"
+#include "OutputLogModule.h"
+#include "Widgets/Docking/SDockTab.h"
+
+#define LOCTEXT_NAMESPACE "UnrealMultiUserUI"
+
+void FOutputLogController::Init(const FConcertComponentInitParams& Params)
+{
+	FGlobalTabmanager::Get()->RegisterTabSpawner(
+			ConcertServerTabs::GetOutputLogTabId(),
+			FOnSpawnTab::CreateRaw(this, &FOutputLogController::SpawnOutputLogTab)
+		)
+		.SetDisplayName(LOCTEXT("OutputLogTabTitle", "Output Log"))
+		.SetIcon(FSlateIcon(FEditorStyle::GetStyleSetName(), "Log.TabIcon"));
+}
+
+TSharedRef<SDockTab> FOutputLogController::SpawnOutputLogTab(const FSpawnTabArgs& Args)
+{
+	FOutputLogCreationParams Params;
+	Params.bCreateDockInLayoutButton = false;
+	Params.SettingsMenuCreationFlags = EOutputLogSettingsMenuFlags::SkipClearOnPie;
+	Params.DefaultCategorySelection = { { "LogSlate", false }, { "LogWindowsTextInputMethodSystem", false } };
+	Params.AllowAsInitialLogCategory = FAllowLogCategoryCallback::CreateLambda([](const FName LogCategoryName)
+	{
+		const FString LogCategoryAsString = LogCategoryName.ToString();
+		return LogCategoryAsString.Contains("Sync") || LogCategoryAsString.Contains("Concert");
+	});
+	
+	return SNew(SDockTab)
+		.Label(LOCTEXT("OutputLogTabTitle", "Output Log"))
+		.TabRole(MajorTab)
+		[
+			FOutputLogModule::Get().MakeOutputLogWidget(Params)
+		];
+}
+
+#undef LOCTEXT_NAMESPACE
diff --git a/Engine/Plugins/Developer/Concert/ConcertUI/ConcertServerUI/Source/ConcertServerUI/Private/Widgets/OutputLog/OutputLogController.h b/Engine/Plugins/Developer/Concert/ConcertUI/ConcertServerUI/Source/ConcertServerUI/Private/Widgets/OutputLog/OutputLogController.h
new file mode 100644
--- /dev/null
+++ b/Engine/Plugins/Developer/Concert/ConcertUI/ConcertServerUI/Source/ConcertServerUI/Private/Widgets/OutputLog/OutputLogController.h
@@ -0,0 +1,23 @@
+﻿// Copyright Epic Games, Inc. All Rights Reserved.
+
+#pragma once
+
+#include "CoreMinimal.h"
+#include "Widgets/IConcertComponent.h"
+
+class FSpawnTabArgs;
+class SDockTab;
+
+class FOutputLogController
+	: public IConcertComponent
+{
+public:
+	
+	//~ Begin IConcertComponent Interface
+	virtual void Init(const FConcertComponentInitParams& Params) override;
+	//~ End IConcertComponent Interface
+
+private:
+	
+	TSharedRef<SDockTab> SpawnOutputLogTab(const FSpawnTabArgs& Args);
+};
diff --git a/Engine/Plugins/Developer/Concert/ConcertUI/ConcertServerUI/Source/ConcertServerUI/Private/Widgets/StatusBar/SConcertStatusBar.cpp b/Engine/Plugins/Developer/Concert/ConcertUI/ConcertServerUI/Source/ConcertServerUI/Private/Widgets/StatusBar/SConcertStatusBar.cpp
new file mode 100644
--- /dev/null
+++ b/Engine/Plugins/Developer/Concert/ConcertUI/ConcertServerUI/Source/ConcertServerUI/Private/Widgets/StatusBar/SConcertStatusBar.cpp
@@ -0,0 +1,13 @@
+﻿// Copyright Epic Games, Inc. All Rights Reserved.
+
+#include "SConcertStatusBar.h"
+
+#include "Widgets/SBoxPanel.h"
+
+void SConcertStatusBar::Construct(const FArguments& InArgs)
+{
+	ChildSlot
+	[
+		SNew(SHorizontalBox)
+	];
+}
diff --git a/Engine/Plugins/Developer/Concert/ConcertUI/ConcertServerUI/Source/ConcertServerUI/Private/Widgets/StatusBar/SConcertStatusBar.h b/Engine/Plugins/Developer/Concert/ConcertUI/ConcertServerUI/Source/ConcertServerUI/Private/Widgets/StatusBar/SConcertStatusBar.h
new file mode 100644
--- /dev/null
+++ b/Engine/Plugins/Developer/Concert/ConcertUI/ConcertServerUI/Source/ConcertServerUI/Private/Widgets/StatusBar/SConcertStatusBar.h
@@ -0,0 +1,17 @@
+﻿// Copyright Epic Games, Inc. All Rights Reserved.
+
+#pragma once
+
+#include "CoreMinimal.h"
+#include "Widgets/DeclarativeSyntaxSupport.h"
+#include "Widgets/SCompoundWidget.h"
+
+class SConcertStatusBar : public SCompoundWidget
+{
+public:
+
+	SLATE_BEGIN_ARGS(SConcertStatusBar) { }
+	SLATE_END_ARGS()
+
+	void Construct(const FArguments& InArgs);
+};
diff --git a/Engine/Plugins/Developer/Concert/ConcertUI/ConcertServerUI/Source/ConcertServerUI/Public/IConcertServerUIModule.h b/Engine/Plugins/Developer/Concert/ConcertUI/ConcertServerUI/Source/ConcertServerUI/Public/IConcertServerUIModule.h
new file mode 100644
--- /dev/null
+++ b/Engine/Plugins/Developer/Concert/ConcertUI/ConcertServerUI/Source/ConcertServerUI/Public/IConcertServerUIModule.h
@@ -0,0 +1,43 @@
+// Copyright Epic Games, Inc. All Rights Reserved.
+
+#pragma once
+
+#include "CoreMinimal.h"
+#include "Modules/ModuleInterface.h"
+#include "Modules/ModuleManager.h"
+
+class IConcertSyncServer;
+struct FConcertSyncServerLoopInitArgs;
+
+/**
+ * Interface for the Concert Sync Server module.
+ */
+class IConcertServerUIModule : public IModuleInterface
+{
+public:
+	/**
+	 * Singleton-like access to this module's interface.  This is just for convenience!
+	 * Beware of calling this during the shutdown phase, though.  Your module might have been unloaded already.
+	 *
+	 * @return Returns singleton instance, loading the module on demand if needed
+	 */
+	static inline IConcertServerUIModule& Get()
+	{
+		static const FName ModuleName = "ConcertServerUI";
+		return FModuleManager::LoadModuleChecked<IConcertServerUIModule>(ModuleName);
+	}
+
+	/**
+	 * Checks to see if this module is loaded and ready.  It is only valid to call Get() during shutdown if IsAvailable() returns true.
+	 *
+	 * @return True if the module is loaded and ready to use
+	 */
+	static inline bool IsAvailable()
+	{
+		static const FName ModuleName = "ConcertServerUI";
+		return FModuleManager::Get().IsModuleLoaded(ModuleName);
+	}
+
+	/** Set-up callbacks for creating Slate UI when the MU server is started. Called before the server loop is initialized. */
+	virtual void InitSlateForServer(FConcertSyncServerLoopInitArgs& InitArgs) = 0;
+};
diff --git a/Engine/Plugins/Developer/Concert/MultiUserClient/Source/MultiUserClient/Private/MultiUserClientModule.cpp b/Engine/Plugins/Developer/Concert/MultiUserClient/Source/MultiUserClient/Private/MultiUserClientModule.cpp
--- a/Engine/Plugins/Developer/Concert/MultiUserClient/Source/MultiUserClient/Private/MultiUserClientModule.cpp
+++ b/Engine/Plugins/Developer/Concert/MultiUserClient/Source/MultiUserClient/Private/MultiUserClientModule.cpp
@@ -1309,10 +1309,10 @@ private:
 			// Unicast endpoint setting
 			EngineConfig->GetString(TEXT("/Script/UdpMessaging.UdpMessagingSettings"), TEXT("UnicastEndpoint"), Setting);
 
-			// if the unicast endpoint port is bound, add 1 to the port for server 
+			// if the unicast endpoint port is bound, add 1 to the port for server
+			const UConcertClientConfig* ClientConfig = GetDefault<UConcertClientConfig>();
 			if (Setting.ParseIntoArray(Settings, TEXT(":"), false) == 2)
 			{
-				const UConcertClientConfig* ClientConfig = GetDefault<UConcertClientConfig>();
 				if (ClientConfig && ClientConfig->ClientSettings.ServerPort != 0)
 				{
 					Setting = FString::Printf(TEXT("%s:%d"), *Settings[0], ClientConfig->ClientSettings.ServerPort);
@@ -1343,6 +1343,11 @@ private:
 					CmdLine += Settings[i];
 				}
 			}
+
+			if (ClientConfig && ClientConfig->bRunWithSlate)
+			{
+				CmdLine += TEXT(" -WITHSLATE");
+			}
 		}
 		return CmdLine;
 	}
diff --git a/Engine/Programs/UnrealMultiUserServer/Config/DefaultEngine.ini b/Engine/Programs/UnrealMultiUserServer/Config/DefaultEngine.ini
--- a/Engine/Programs/UnrealMultiUserServer/Config/DefaultEngine.ini
+++ b/Engine/Programs/UnrealMultiUserServer/Config/DefaultEngine.ini
@@ -4,3 +4,4 @@ GameNameShort=UCS
 [Plugins]
 +ProgramEnabledPlugins="UdpMessaging"
 +ProgramEnabledPlugins="ConcertSyncServer"
++ProgramEnabledPlugins="ConcertServerUI"
diff --git a/Engine/Source/Programs/UnrealMultiUserServer/Private/UnrealMultiUserServerRun.cpp b/Engine/Source/Programs/UnrealMultiUserServer/Private/UnrealMultiUserServerRun.cpp
--- a/Engine/Source/Programs/UnrealMultiUserServer/Private/UnrealMultiUserServerRun.cpp
+++ b/Engine/Source/Programs/UnrealMultiUserServer/Private/UnrealMultiUserServerRun.cpp
@@ -4,6 +4,38 @@
 
 #include "ConcertSettings.h"
 #include "ConcertSyncServerLoop.h"
+#include "IConcertServerUIModule.h"
+
+#include "Misc/Parse.h"
+
+namespace UE::UnrealMultiUserServer
+{
+	static void OptionallySetupSlate(int ArgC, TCHAR* ArgV[], FConcertSyncServerLoopInitArgs& ServerLoopInitArgs)
+	{
+		const FString CommandLine = FCommandLine::BuildFromArgV(nullptr, ArgC, ArgV, nullptr);
+
+		bool bUseSlate = false;
+		bUseSlate |= FParse::Bool(*CommandLine, TEXT("-WITHSLATE"), bUseSlate);
+		FParse::Bool(*CommandLine, TEXT("-WITHSLATE="), bUseSlate);
+		
+		if (bUseSlate)
+		{
+			ServerLoopInitArgs.bShowConsole = false;
+			ServerLoopInitArgs.PreInitServerLoop.AddLambda([&ServerLoopInitArgs]()	
+			{
+				TSharedPtr<IPlugin> Plugin = IPluginManager::Get().FindPlugin(TEXT("ConcertServerUI"));
+				if (!Plugin || !Plugin->IsEnabled())
+				{
+					UE_LOG(LogSyncServer, Error, TEXT("The 'ConcertServerUI' plugin is disabled."));
+				}
+				else
+				{
+					IConcertServerUIModule::Get().InitSlateForServer(ServerLoopInitArgs);
+				}
+			});
+		}
+	}
+}
 
 int32 RunUnrealMultiUserServer(int ArgC, TCHAR* ArgV[])
 {
@@ -27,5 +59,6 @@ int32 RunUnrealMultiUserServer(int ArgC, TCHAR* ArgV[])
 		return ServerConfig;
 	};
 
+	UE::UnrealMultiUserServer::OptionallySetupSlate(ArgC, ArgV, ServerLoopInitArgs);
 	return ConcertSyncServerLoop(ArgC, ArgV, ServerLoopInitArgs);
 }
diff --git a/Engine/Source/Programs/UnrealMultiUserServer/UnrealMultiUserServer.Build.cs b/Engine/Source/Programs/UnrealMultiUserServer/UnrealMultiUserServer.Build.cs
--- a/Engine/Source/Programs/UnrealMultiUserServer/UnrealMultiUserServer.Build.cs
+++ b/Engine/Source/Programs/UnrealMultiUserServer/UnrealMultiUserServer.Build.cs
@@ -22,6 +22,7 @@ public class UnrealMultiUserServer : ModuleRules
 				"Launch",
 				"ConcertSyncCore",
 				"ConcertSyncServer",
+				"ConcertServerUI",
 			}
 		);
 
@@ -34,6 +35,7 @@ public class UnrealMultiUserServer : ModuleRules
 		DynamicallyLoadedModuleNames.AddRange(
 			new string[] {
 				"ConcertSyncServer",
+				"ConcertServerUI"
 			}
 		);
 	}
