Fix non-unity compile issues

diff --git a/Engine/Plugins/Developer/Concert/ConcertUI/ConcertServerUI/Source/ConcertServerUI/Private/ConcertServerUIModule.cpp b/Engine/Plugins/Developer/Concert/ConcertUI/ConcertServerUI/Source/ConcertServerUI/Private/ConcertServerUIModule.cpp
--- a/Engine/Plugins/Developer/Concert/ConcertUI/ConcertServerUI/Source/ConcertServerUI/Private/ConcertServerUIModule.cpp
+++ b/Engine/Plugins/Developer/Concert/ConcertUI/ConcertServerUI/Source/ConcertServerUI/Private/ConcertServerUIModule.cpp
@@ -5,12 +5,12 @@
 #include "ConcertConsoleCommandExecutor.h"
 #include "ConcertServerStyle.h"
 #include "ConcertSyncServerLoopInitArgs.h"
-#include "OutputLogModule.h"
 #include "Widgets/ConcertServerWindowController.h"
 
+#include "Features/IModularFeatures.h"
 #include "Framework/Application/SlateApplication.h"
 #include "StandaloneRenderer.h"
-#include "Features/IModularFeatures.h"
+#include "Misc/ConfigCacheIni.h"
 
 #define LOCTEXT_NAMESPACE "UnrealMultiUserUI"
 
diff --git a/Engine/Plugins/Developer/Concert/ConcertUI/ConcertServerUI/Source/ConcertServerUI/Private/Widgets/Browser/ConcertServerSessionBrowserController.cpp b/Engine/Plugins/Developer/Concert/ConcertUI/ConcertServerUI/Source/ConcertServerUI/Private/Widgets/Browser/ConcertServerSessionBrowserController.cpp
--- a/Engine/Plugins/Developer/Concert/ConcertUI/ConcertServerUI/Source/ConcertServerUI/Private/Widgets/Browser/ConcertServerSessionBrowserController.cpp
+++ b/Engine/Plugins/Developer/Concert/ConcertUI/ConcertServerUI/Source/ConcertServerUI/Private/Widgets/Browser/ConcertServerSessionBrowserController.cpp
@@ -3,8 +3,11 @@
 #include "ConcertServerSessionBrowserController.h"
 
 #include "ConcertServerStyle.h"
+#include "Framework/Docking/TabManager.h"
 #include "SConcertServerSessionBrowser.h"
+#include "Textures/SlateIcon.h"
 #include "Widgets/ConcertServerTabs.h"
+#include "Widgets/Docking/SDockTab.h"
 
 #define LOCTEXT_NAMESPACE "UnrealMultiUserUI"
 
diff --git a/Engine/Plugins/Developer/Concert/ConcertUI/ConcertServerUI/Source/ConcertServerUI/Private/Widgets/ConcertServerTabs.h b/Engine/Plugins/Developer/Concert/ConcertUI/ConcertServerUI/Source/ConcertServerUI/Private/Widgets/ConcertServerTabs.h
--- a/Engine/Plugins/Developer/Concert/ConcertUI/ConcertServerUI/Source/ConcertServerUI/Private/Widgets/ConcertServerTabs.h
+++ b/Engine/Plugins/Developer/Concert/ConcertUI/ConcertServerUI/Source/ConcertServerUI/Private/Widgets/ConcertServerTabs.h
@@ -6,6 +6,6 @@
 
 namespace ConcertServerTabs
 {
-	static FName GetSessionBrowserTabId();
-	static FName GetOutputLogTabId();
+	FName GetSessionBrowserTabId();
+	FName GetOutputLogTabId();
 }
diff --git a/Engine/Plugins/Developer/Concert/ConcertUI/ConcertServerUI/Source/ConcertServerUI/Private/Widgets/ConcertServerWindowController.cpp b/Engine/Plugins/Developer/Concert/ConcertUI/ConcertServerUI/Source/ConcertServerUI/Private/Widgets/ConcertServerWindowController.cpp
--- a/Engine/Plugins/Developer/Concert/ConcertUI/ConcertServerUI/Source/ConcertServerUI/Private/Widgets/ConcertServerWindowController.cpp
+++ b/Engine/Plugins/Developer/Concert/ConcertUI/ConcertServerUI/Source/ConcertServerUI/Private/Widgets/ConcertServerWindowController.cpp
@@ -6,10 +6,12 @@
 #include "Browser/ConcertServerSessionBrowserController.h"
 #include "OutputLog/OutputLogController.h"
 
+#include "Framework/Application/SlateApplication.h"
 #include "Framework/Docking/TabManager.h"
 #include "Framework/Docking/LayoutService.h"
 #include "Framework/Notifications/NotificationManager.h"
 #include "HAL/PlatformApplicationMisc.h"
+#include "Misc/ConfigCacheIni.h"
 #include "OutputLog/Public/OutputLogModule.h"
 #include "Widgets/SWindow.h"
 
diff --git a/Engine/Plugins/Developer/Concert/ConcertUI/ConcertServerUI/Source/ConcertServerUI/Private/Widgets/ConcertServerWindowController.h b/Engine/Plugins/Developer/Concert/ConcertUI/ConcertServerUI/Source/ConcertServerUI/Private/Widgets/ConcertServerWindowController.h
--- a/Engine/Plugins/Developer/Concert/ConcertUI/ConcertServerUI/Source/ConcertServerUI/Private/Widgets/ConcertServerWindowController.h
+++ b/Engine/Plugins/Developer/Concert/ConcertUI/ConcertServerUI/Source/ConcertServerUI/Private/Widgets/ConcertServerWindowController.h
@@ -3,10 +3,12 @@
 #pragma once
 
 #include "CoreMinimal.h"
+#include "Framework/Docking/TabManager.h"
 
 class IConcertSyncServer;
 class FConcertServerSessionBrowserController;
 class FOutputLogController;
+class SWindow;
 
 struct FConcertServerWindowInitParams
 {
diff --git a/Engine/Plugins/Developer/Concert/ConcertUI/ConcertServerUI/Source/ConcertServerUI/Private/Widgets/IConcertComponent.h b/Engine/Plugins/Developer/Concert/ConcertUI/ConcertServerUI/Source/ConcertServerUI/Private/Widgets/IConcertComponent.h
--- a/Engine/Plugins/Developer/Concert/ConcertUI/ConcertServerUI/Source/ConcertServerUI/Private/Widgets/IConcertComponent.h
+++ b/Engine/Plugins/Developer/Concert/ConcertUI/ConcertServerUI/Source/ConcertServerUI/Private/Widgets/IConcertComponent.h
@@ -4,6 +4,8 @@
 
 #include "CoreMinimal.h"
 
+class IConcertSyncServer;
+
 struct FConcertComponentInitParams
 {
 	TSharedRef<IConcertSyncServer> Server;
