Implement status bar for multi user server

diff --git a/Engine/Plugins/Developer/Concert/ConcertUI/ConcertServerUI/Source/ConcertServerUI/Private/Widgets/Browser/SConcertServerSessionBrowser.cpp b/Engine/Plugins/Developer/Concert/ConcertUI/ConcertServerUI/Source/ConcertServerUI/Private/Widgets/Browser/SConcertServerSessionBrowser.cpp
--- a/Engine/Plugins/Developer/Concert/ConcertUI/ConcertServerUI/Source/ConcertServerUI/Private/Widgets/Browser/SConcertServerSessionBrowser.cpp
+++ b/Engine/Plugins/Developer/Concert/ConcertUI/ConcertServerUI/Source/ConcertServerUI/Private/Widgets/Browser/SConcertServerSessionBrowser.cpp
@@ -4,9 +4,11 @@
 
 #include "SPositiveActionButton.h"
 #include "OutputLog/Public/OutputLogModule.h"
+#include "Widgets/ConcertServerTabs.h"
 #include "Widgets/Input/SSearchBox.h"
 #include "Widgets/Layout/SBorder.h"
 #include "Widgets/SBoxPanel.h"
+#include "Widgets/StatusBar/SConcertStatusBar.h"
 
 #define LOCTEXT_NAMESPACE "UnrealMultiUserUI"
 
@@ -33,6 +35,14 @@ void SConcertServerSessionBrowser::Construct(const FArguments& InArgs, TWeakPtr<
 			[
 				MakeSessionTableView()
 			]
+
+			+SVerticalBox::Slot()
+			.AutoHeight()
+			.HAlign(HAlign_Fill)
+			.VAlign(VAlign_Bottom)
+			[
+				SNew(SConcertStatusBar, ConcertServerTabs::GetSessionBrowserTabId())
+			]
 		]
 	];
 }
diff --git a/Engine/Plugins/Developer/Concert/ConcertUI/ConcertServerUI/Source/ConcertServerUI/Private/Widgets/ConcertServerTabs.cpp b/Engine/Plugins/Developer/Concert/ConcertUI/ConcertServerUI/Source/ConcertServerUI/Private/Widgets/ConcertServerTabs.cpp
--- a/Engine/Plugins/Developer/Concert/ConcertUI/ConcertServerUI/Source/ConcertServerUI/Private/Widgets/ConcertServerTabs.cpp
+++ b/Engine/Plugins/Developer/Concert/ConcertUI/ConcertServerUI/Source/ConcertServerUI/Private/Widgets/ConcertServerTabs.cpp
@@ -6,10 +6,4 @@ FName ConcertServerTabs::GetSessionBrowserTabId()
 {
 	static FName SessionBrowserTabId("SessionBrowserTabId");
 	return SessionBrowserTabId;
-}
-
-FName ConcertServerTabs::GetOutputLogTabId()
-{
-	static FName OutputLogTabId("OutputLogTabId");
-	return OutputLogTabId;
 }
\ No newline at end of file
diff --git a/Engine/Plugins/Developer/Concert/ConcertUI/ConcertServerUI/Source/ConcertServerUI/Private/Widgets/ConcertServerTabs.h b/Engine/Plugins/Developer/Concert/ConcertUI/ConcertServerUI/Source/ConcertServerUI/Private/Widgets/ConcertServerTabs.h
--- a/Engine/Plugins/Developer/Concert/ConcertUI/ConcertServerUI/Source/ConcertServerUI/Private/Widgets/ConcertServerTabs.h
+++ b/Engine/Plugins/Developer/Concert/ConcertUI/ConcertServerUI/Source/ConcertServerUI/Private/Widgets/ConcertServerTabs.h
@@ -7,5 +7,4 @@
 namespace ConcertServerTabs
 {
 	FName GetSessionBrowserTabId();
-	FName GetOutputLogTabId();
 }
diff --git a/Engine/Plugins/Developer/Concert/ConcertUI/ConcertServerUI/Source/ConcertServerUI/Private/Widgets/ConcertServerWindowController.cpp b/Engine/Plugins/Developer/Concert/ConcertUI/ConcertServerUI/Source/ConcertServerUI/Private/Widgets/ConcertServerWindowController.cpp
--- a/Engine/Plugins/Developer/Concert/ConcertUI/ConcertServerUI/Source/ConcertServerUI/Private/Widgets/ConcertServerWindowController.cpp
+++ b/Engine/Plugins/Developer/Concert/ConcertUI/ConcertServerUI/Source/ConcertServerUI/Private/Widgets/ConcertServerWindowController.cpp
@@ -4,7 +4,6 @@
 
 #include "ConcertServerTabs.h"
 #include "Browser/ConcertServerSessionBrowserController.h"
-#include "OutputLog/OutputLogController.h"
 
 #include "Framework/Application/SlateApplication.h"
 #include "Framework/Docking/TabManager.h"
@@ -21,7 +20,6 @@ FConcertServerWindowController::FConcertServerWindowController(const FConcertSer
 	: MultiUserServerLayoutIni(Params.MultiUserServerLayoutIni)
 {
 	SessionBrowserController = MakeShared<FConcertServerSessionBrowserController>();
-	OutputLogController = MakeShared<FOutputLogController>();
 	InitComponents(Params);
 
 	FDisplayMetrics DisplayMetrics;
@@ -56,7 +54,6 @@ FConcertServerWindowController::FConcertServerWindowController(const FConcertSer
 		(
 			FTabManager::NewStack()
 				->AddTab(ConcertServerTabs::GetSessionBrowserTabId(), ETabState::OpenedTab)
-				->AddTab(ConcertServerTabs::GetOutputLogTabId(), ETabState::OpenedTab)
 				->SetForegroundTab(ConcertServerTabs::GetSessionBrowserTabId())
 		)
 	);
@@ -75,7 +72,6 @@ void FConcertServerWindowController::InitComponents(const FConcertServerWindowIn
 {
 	const FConcertComponentInitParams Params { WindowInitParams.Server };
 	SessionBrowserController->Init(Params);
-	OutputLogController->Init(Params);
 }
 
 void FConcertServerWindowController::OnWindowClosed(const TSharedRef<SWindow>& Window)
diff --git a/Engine/Plugins/Developer/Concert/ConcertUI/ConcertServerUI/Source/ConcertServerUI/Private/Widgets/ConcertServerWindowController.h b/Engine/Plugins/Developer/Concert/ConcertUI/ConcertServerUI/Source/ConcertServerUI/Private/Widgets/ConcertServerWindowController.h
--- a/Engine/Plugins/Developer/Concert/ConcertUI/ConcertServerUI/Source/ConcertServerUI/Private/Widgets/ConcertServerWindowController.h
+++ b/Engine/Plugins/Developer/Concert/ConcertUI/ConcertServerUI/Source/ConcertServerUI/Private/Widgets/ConcertServerWindowController.h
@@ -41,8 +41,6 @@ private:
 
 	/** Manages the session browser */
 	TSharedPtr<FConcertServerSessionBrowserController> SessionBrowserController;
-	/** Manages tab for output log */
-	TSharedPtr<FOutputLogController> OutputLogController;
 	
 	void InitComponents(const FConcertServerWindowInitParams& WindowInitParams) const;
 
diff --git a/Engine/Plugins/Developer/Concert/ConcertUI/ConcertServerUI/Source/ConcertServerUI/Private/Widgets/OutputLog/OutputLogController.cpp b/Engine/Plugins/Developer/Concert/ConcertUI/ConcertServerUI/Source/ConcertServerUI/Private/Widgets/OutputLog/OutputLogController.cpp
deleted file mode 100644
diff --git a/Engine/Plugins/Developer/Concert/ConcertUI/ConcertServerUI/Source/ConcertServerUI/Private/Widgets/OutputLog/OutputLogController.h b/Engine/Plugins/Developer/Concert/ConcertUI/ConcertServerUI/Source/ConcertServerUI/Private/Widgets/OutputLog/OutputLogController.h
deleted file mode 100644
diff --git a/Engine/Plugins/Developer/Concert/ConcertUI/ConcertServerUI/Source/ConcertServerUI/Private/Widgets/StatusBar/SConcertStatusBar.cpp b/Engine/Plugins/Developer/Concert/ConcertUI/ConcertServerUI/Source/ConcertServerUI/Private/Widgets/StatusBar/SConcertStatusBar.cpp
--- a/Engine/Plugins/Developer/Concert/ConcertUI/ConcertServerUI/Source/ConcertServerUI/Private/Widgets/StatusBar/SConcertStatusBar.cpp
+++ b/Engine/Plugins/Developer/Concert/ConcertUI/ConcertServerUI/Source/ConcertServerUI/Private/Widgets/StatusBar/SConcertStatusBar.cpp
@@ -2,12 +2,131 @@
 
 #include "SConcertStatusBar.h"
 
-#include "Widgets/SBoxPanel.h"
+#include "OutputLogModule.h"
+#include "SWidgetDrawer.h"
+#include "Widgets/Input/SMultiLineEditableTextBox.h"
+#include "Widgets/Layout/SBox.h"
 
-void SConcertStatusBar::Construct(const FArguments& InArgs)
+#define LOCTEXT_NAMESPACE "UnrealMultiUserUI"
+
+namespace UE::ConcertServerUI::Private
+{
+	const FName OutputLogId("OutputLog");
+	
+	class FStatusBarSingleton
+	{
+		TSharedPtr<SWidget> StatusBarOutputLog;
+		TArray<TWeakPtr<SWidgetDrawer>> StatusBars;
+		
+		TSharedRef<SWidget> OnGetOutputLog()
+		{
+			if (!StatusBarOutputLog)
+			{
+				StatusBarOutputLog = FOutputLogModule::Get().MakeOutputLogDrawerWidget(FSimpleDelegate());
+			}
+
+			return StatusBarOutputLog.ToSharedRef();
+		}
+		
+		void OnOutputLogOpened(FName StatusBarWithDrawerName)
+		{
+			// Dismiss all other open drawers - StatusBarOutputLog is shared and shouldn't be in the layout twice
+			for (TWeakPtr<SWidgetDrawer> WidgetDrawer : StatusBars)
+			{
+				if (WidgetDrawer.IsValid())
+				{
+					TSharedPtr<SWidgetDrawer> PinnedDrawer = WidgetDrawer.Pin();
+					if (StatusBarWithDrawerName != PinnedDrawer->GetDrawerName() || PinnedDrawer->IsAnyOtherDrawerOpened(OutputLogId))
+					{
+						PinnedDrawer->CloseDrawerImmediately();
+					}
+				}
+			}
+			
+			FOutputLogModule::Get().FocusOutputLogConsoleBox(StatusBarOutputLog.ToSharedRef());
+		}
+		
+		void OnOutputLogDismised(const TSharedPtr<SWidget>& NewlyFocusedWidget)
+		{}
+
+		void PreShutdownSlate()
+		{
+			StatusBarOutputLog.Reset();
+		}
+
+	public:
+
+		void Init(TSharedRef<SWidgetDrawer> WidgetDrawer, FWidgetDrawerConfig& OutputLogDrawer)
+		{
+			if (!FSlateApplication::Get().OnPreShutdown().IsBoundToObject(this))
+			{
+				// Destroying StatusBarOutputLog in ~FStatusBarSingleton is too late: it causes a crash
+				FSlateApplication::Get().OnPreShutdown().AddRaw(this, &FStatusBarSingleton::PreShutdownSlate);
+			}
+			
+			StatusBars.Add(WidgetDrawer);
+			
+			OutputLogDrawer.GetDrawerContentDelegate.BindRaw(this, &FStatusBarSingleton::OnGetOutputLog);
+			OutputLogDrawer.OnDrawerOpenedDelegate.BindRaw(this, &FStatusBarSingleton::OnOutputLogOpened);
+			OutputLogDrawer.OnDrawerDismissedDelegate.BindRaw(this, &FStatusBarSingleton::OnOutputLogDismised);
+		}
+
+		void Remove(TSharedRef<SWidgetDrawer> WidgetDrawer)
+		{
+			StatusBars.RemoveSingle(WidgetDrawer);
+		}
+		
+	} GStatusBarManager;
+}
+
+SConcertStatusBar::~SConcertStatusBar()
+{
+	UE::ConcertServerUI::Private::GStatusBarManager.Remove(WidgetDrawer.ToSharedRef());
+}
+
+void SConcertStatusBar::Construct(const FArguments& InArgs, FName StatusBarId)
 {
 	ChildSlot
 	[
-		SNew(SHorizontalBox)
+		SNew(SBox)
+		.HeightOverride(FAppStyle::Get().GetFloat("StatusBar.Height"))
+		[
+			MakeWidgetDrawer(StatusBarId)
+		]
 	];
 }
+
+TSharedRef<SWidgetDrawer> SConcertStatusBar::MakeWidgetDrawer(FName StatusBarId)
+{
+	using namespace UE::ConcertServerUI::Private;
+	
+	WidgetDrawer = SNew(SWidgetDrawer, StatusBarId);
+
+	TSharedPtr<SMultiLineEditableTextBox> ConsoleEditBox;
+	FSimpleDelegate OnConsoleClosed;
+	FSimpleDelegate OnConsoleCommandExecuted;
+	const TSharedRef<SWidget> OutputLog = 
+			SNew(SBorder)
+			.BorderImage(FAppStyle::Get().GetBrush("Brushes.Panel"))
+			.VAlign(VAlign_Center)
+			.Padding(FMargin(6.0f, 0.0f))
+			[
+				SNew(SBox)
+				.WidthOverride(350.f)
+				[
+					FOutputLogModule::Get().MakeConsoleInputBox(ConsoleEditBox, OnConsoleClosed, OnConsoleCommandExecuted)
+				]
+			];
+	
+	FWidgetDrawerConfig OutputLogDrawer(OutputLogId);
+	GStatusBarManager.Init(WidgetDrawer.ToSharedRef(), OutputLogDrawer);
+	OutputLogDrawer.CustomWidget = OutputLog;
+
+	OutputLogDrawer.ButtonText = LOCTEXT("StatusBar_OutputLogButton", "Output Log");
+	OutputLogDrawer.Icon = FAppStyle::Get().GetBrush("Log.TabIcon");
+	WidgetDrawer->RegisterDrawer(MoveTemp(OutputLogDrawer));
+
+	return WidgetDrawer.ToSharedRef();
+}
+
+#undef LOCTEXT_NAMESPACE
diff --git a/Engine/Plugins/Developer/Concert/ConcertUI/ConcertServerUI/Source/ConcertServerUI/Private/Widgets/StatusBar/SConcertStatusBar.h b/Engine/Plugins/Developer/Concert/ConcertUI/ConcertServerUI/Source/ConcertServerUI/Private/Widgets/StatusBar/SConcertStatusBar.h
--- a/Engine/Plugins/Developer/Concert/ConcertUI/ConcertServerUI/Source/ConcertServerUI/Private/Widgets/StatusBar/SConcertStatusBar.h
+++ b/Engine/Plugins/Developer/Concert/ConcertUI/ConcertServerUI/Source/ConcertServerUI/Private/Widgets/StatusBar/SConcertStatusBar.h
@@ -3,6 +3,7 @@
 #pragma once
 
 #include "CoreMinimal.h"
+#include "SWidgetDrawer.h"
 #include "Widgets/DeclarativeSyntaxSupport.h"
 #include "Widgets/SCompoundWidget.h"
 
@@ -12,6 +13,14 @@ public:
 
 	SLATE_BEGIN_ARGS(SConcertStatusBar) { }
 	SLATE_END_ARGS()
+	
+	~SConcertStatusBar();
 
-	void Construct(const FArguments& InArgs);
+	void Construct(const FArguments& InArgs, FName StatusBarId);
+
+private:
+
+	TSharedPtr<SWidgetDrawer> WidgetDrawer;
+
+	TSharedRef<SWidgetDrawer> MakeWidgetDrawer(FName StatusBarId);
 };
