Add session browser to multi user server

diff --git a/Engine/Plugins/Developer/Concert/ConcertMain/Source/Concert/Private/ConcertServer.cpp b/Engine/Plugins/Developer/Concert/ConcertMain/Source/Concert/Private/ConcertServer.cpp
--- a/Engine/Plugins/Developer/Concert/ConcertMain/Source/Concert/Private/ConcertServer.cpp
+++ b/Engine/Plugins/Developer/Concert/ConcertMain/Source/Concert/Private/ConcertServer.cpp
@@ -623,7 +623,7 @@ bool FConcertServer::DestroySession(const FGuid& SessionId, FText& OutFailureRea
 
 TArray<FConcertSessionClientInfo> FConcertServer::GetSessionClients(const FGuid& SessionId) const
 {
-	TSharedPtr<IConcertServerSession> ServerSession = GetSession(SessionId);
+	TSharedPtr<IConcertServerSession> ServerSession = GetLiveSession(SessionId);
 	if (ServerSession)
 	{
 		return ServerSession->GetSessionClients();
@@ -631,7 +631,7 @@ TArray<FConcertSessionClientInfo> FConcertServer::GetSessionClients(const FGuid&
 	return TArray<FConcertSessionClientInfo>();
 }
 
-TArray<FConcertSessionInfo> FConcertServer::GetSessionsInfo() const
+TArray<FConcertSessionInfo> FConcertServer::GetLiveSessionInfos() const
 {
 	TArray<FConcertSessionInfo> SessionsInfo;
 	SessionsInfo.Reserve(LiveSessions.Num());
@@ -642,7 +642,18 @@ TArray<FConcertSessionInfo> FConcertServer::GetSessionsInfo() const
 	return SessionsInfo;
 }
 
-TArray<TSharedPtr<IConcertServerSession>> FConcertServer::GetSessions() const
+TArray<FConcertSessionInfo> FConcertServer::GetArchivedSessionInfos() const
+{
+	TArray<FConcertSessionInfo> SessionsInfo;
+	SessionsInfo.Reserve(ArchivedSessions.Num());
+	for (auto& SessionPair : ArchivedSessions)
+	{
+		SessionsInfo.Add(SessionPair.Value);
+	}
+	return SessionsInfo;
+}
+
+TArray<TSharedPtr<IConcertServerSession>> FConcertServer::GetLiveSessions() const
 {
 	TArray<TSharedPtr<IConcertServerSession>> SessionsArray;
 	SessionsArray.Reserve(LiveSessions.Num());
@@ -653,11 +664,17 @@ TArray<TSharedPtr<IConcertServerSession>> FConcertServer::GetSessions() const
 	return SessionsArray;
 }
 
-TSharedPtr<IConcertServerSession> FConcertServer::GetSession(const FGuid& SessionId) const
+TSharedPtr<IConcertServerSession> FConcertServer::GetLiveSession(const FGuid& SessionId) const
 {
 	return LiveSessions.FindRef(SessionId);
 }
 
+TOptional<FConcertSessionInfo> FConcertServer::GetArchivedSessionInfo(const FGuid& SessionId) const
+{
+	const FConcertSessionInfo* SessionInfo = ArchivedSessions.Find(SessionId);
+	return SessionInfo ? *SessionInfo : TOptional<FConcertSessionInfo>{};
+}
+
 const FString& FConcertServer::GetSessionRepositoriesRootDir() const
 {
 	return SessionRepositoryRootDir;
@@ -1049,7 +1066,7 @@ TFuture<FConcertAdmin_SessionInfoResponse> FConcertServer::HandleFindSessionRequ
 	FConcertAdmin_SessionInfoResponse ResponseData;
 
 	// Find the session requested
-	TSharedPtr<IConcertServerSession> ServerSession = GetSession(Message->SessionId);
+	TSharedPtr<IConcertServerSession> ServerSession = GetLiveSession(Message->SessionId);
 	const TCHAR* ServerSessionNamePtr = ServerSession ? *ServerSession->GetName() : TEXT("<unknown>");
 	if (CanJoinSession(ServerSession, Message->SessionSettings, Message->VersionInfo, &ResponseData.Reason))
 	{
@@ -1110,7 +1127,7 @@ TFuture<FConcertAdmin_ArchiveSessionResponse> FConcertServer::HandleArchiveSessi
 	FConcertAdmin_ArchiveSessionResponse ResponseData;
 
 	// Find the session requested.
-	TSharedPtr<IConcertServerSession> ServerSession = GetSession(Message->SessionId);
+	TSharedPtr<IConcertServerSession> ServerSession = GetLiveSession(Message->SessionId);
 	ResponseData.SessionId = Message->SessionId;
 	ResponseData.SessionName = ServerSession ? ServerSession->GetName() : TEXT("<unknown>");
 	if (ServerSession)
@@ -1153,7 +1170,7 @@ FConcertAdmin_RenameSessionResponse FConcertServer::RenameSessionInternal(const
 	ResponseData.SessionId = Request.SessionId;
 	ResponseData.ResponseCode = EConcertResponseCode::Failed;
 
-	if (TSharedPtr<IConcertServerSession> ServerSession = GetSession(Request.SessionId)) // Live session?
+	if (TSharedPtr<IConcertServerSession> ServerSession = GetLiveSession(Request.SessionId)) // Live session?
 	{
 		ResponseData.OldName = ServerSession->GetName();
 
@@ -1219,7 +1236,7 @@ FConcertAdmin_DeleteSessionResponse FConcertServer::DeleteSessionInternal(const
 	ResponseData.SessionId = Request.SessionId;
 	ResponseData.ResponseCode = EConcertResponseCode::Failed;
 
-	if (TSharedPtr<IConcertServerSession> ServerSession = GetSession(Request.SessionId)) // Live session?
+	if (TSharedPtr<IConcertServerSession> ServerSession = GetLiveSession(Request.SessionId)) // Live session?
 	{
 		ResponseData.SessionName = ServerSession->GetName();
 
@@ -1273,7 +1290,7 @@ TFuture<FConcertAdmin_GetAllSessionsResponse> FConcertServer::HandleGetAllSessio
 	const FConcertAdmin_GetAllSessionsRequest* Message = Context.GetMessage<FConcertAdmin_GetAllSessionsRequest>();
 
 	FConcertAdmin_GetAllSessionsResponse ResponseData;
-	ResponseData.LiveSessions = GetSessionsInfo();
+	ResponseData.LiveSessions = GetLiveSessionInfos();
 	for (const auto& ArchivedSessionPair : ArchivedSessions)
 	{
 		ResponseData.ArchivedSessions.Add(ArchivedSessionPair.Value);
@@ -1287,7 +1304,7 @@ TFuture<FConcertAdmin_GetSessionsResponse> FConcertServer::HandleGetLiveSessions
 	const FConcertAdmin_GetLiveSessionsRequest* Message = Context.GetMessage<FConcertAdmin_GetLiveSessionsRequest>();
 
 	FConcertAdmin_GetSessionsResponse ResponseData;
-	ResponseData.Sessions = GetSessionsInfo();
+	ResponseData.Sessions = GetLiveSessionInfos();
 	
 	return FConcertAdmin_GetSessionsResponse::AsFuture(MoveTemp(ResponseData));
 }
diff --git a/Engine/Plugins/Developer/Concert/ConcertMain/Source/Concert/Private/ConcertServer.h b/Engine/Plugins/Developer/Concert/ConcertMain/Source/Concert/Private/ConcertServer.h
--- a/Engine/Plugins/Developer/Concert/ConcertMain/Source/Concert/Private/ConcertServer.h
+++ b/Engine/Plugins/Developer/Concert/ConcertMain/Source/Concert/Private/ConcertServer.h
@@ -34,9 +34,11 @@ public:
 	virtual FGuid GetArchivedSessionIdByName(const FString& InName) const override;
 
 	virtual FConcertSessionInfo CreateSessionInfo() const override;
-	virtual TArray<FConcertSessionInfo> GetSessionsInfo() const override;
-	virtual TArray<TSharedPtr<IConcertServerSession>> GetSessions() const override;
-	virtual TSharedPtr<IConcertServerSession> GetSession(const FGuid& SessionId) const override;
+	virtual TArray<FConcertSessionInfo> GetLiveSessionInfos() const override;
+	virtual TArray<FConcertSessionInfo> GetArchivedSessionInfos() const override;
+	virtual TArray<TSharedPtr<IConcertServerSession>> GetLiveSessions() const override;
+	virtual TSharedPtr<IConcertServerSession> GetLiveSession(const FGuid& SessionId) const override;
+	virtual TOptional<FConcertSessionInfo> GetArchivedSessionInfo(const FGuid& SessionId) const override;
 	virtual TSharedPtr<IConcertServerSession> CreateSession(const FConcertSessionInfo& SessionInfo, FText& OutFailureReason) override;
 	virtual TSharedPtr<IConcertServerSession> RestoreSession(const FGuid& SessionId, const FConcertSessionInfo& SessionInfo, const FConcertSessionFilter& SessionFilter, FText& OutFailureReason) override;
 	virtual TSharedPtr<IConcertServerSession> CopySession(const FGuid& SrcSessionId, const FConcertSessionInfo& NewSessionInfo, const FConcertSessionFilter& SessionFilter, FText& OutFailureReason) override;
diff --git a/Engine/Plugins/Developer/Concert/ConcertMain/Source/Concert/Public/IConcertServer.h b/Engine/Plugins/Developer/Concert/ConcertMain/Source/Concert/Public/IConcertServer.h
--- a/Engine/Plugins/Developer/Concert/ConcertMain/Source/Concert/Public/IConcertServer.h
+++ b/Engine/Plugins/Developer/Concert/ConcertMain/Source/Concert/Public/IConcertServer.h
@@ -74,22 +74,34 @@ public:
 	virtual FConcertSessionInfo CreateSessionInfo() const = 0;
 
 	/**
-	 * Get the sessions information list
+	 * Get the live session information list
 	 */
-	virtual	TArray<FConcertSessionInfo> GetSessionsInfo() const = 0;
+	virtual	TArray<FConcertSessionInfo> GetLiveSessionInfos() const = 0;
 
 	/**
-	 * Get all server sessions
+	 * Gets the archived session information list
+	 */
+	virtual TArray<FConcertSessionInfo> GetArchivedSessionInfos() const = 0;
+
+	/**
+	 * Get all live server sessions
 	 * @return array of server sessions
 	 */
-	virtual TArray<TSharedPtr<IConcertServerSession>> GetSessions() const = 0;
+	virtual TArray<TSharedPtr<IConcertServerSession>> GetLiveSessions() const = 0;
 
 	/**
-	 * Get a server session
+	 * Get a live server session
 	 * @param SessionId The ID of the session we want
 	 * @return the server session or an invalid pointer if no session was found
 	 */
-	virtual TSharedPtr<IConcertServerSession> GetSession(const FGuid& SessionId) const = 0;
+	virtual TSharedPtr<IConcertServerSession> GetLiveSession(const FGuid& SessionId) const = 0;
+
+	/**
+	 * Gets an archived session info
+	 * @param SessionId The ID of the session we want
+	 * @return The info of the archived session if found
+	 */
+	virtual TOptional<FConcertSessionInfo> GetArchivedSessionInfo(const FGuid& SessionId) const = 0;
 
 	/** 
 	 * Create a new Concert server session based on the passed session info
diff --git a/Engine/Plugins/Developer/Concert/ConcertUI/ConcertServerUI/ConcertServerUI.uplugin b/Engine/Plugins/Developer/Concert/ConcertUI/ConcertServerUI/ConcertServerUI.uplugin
--- a/Engine/Plugins/Developer/Concert/ConcertUI/ConcertServerUI/ConcertServerUI.uplugin
+++ b/Engine/Plugins/Developer/Concert/ConcertUI/ConcertServerUI/ConcertServerUI.uplugin
@@ -29,6 +29,10 @@
 			"Name": "ConcertMain",
 			"Enabled": true
 		},
+        {
+          "Name": "ConcertSharedSlate",
+          "Enabled": true
+        },
 		{
 			"Name": "ConcertSyncCore",
 			"Enabled": true
diff --git a/Engine/Plugins/Developer/Concert/ConcertUI/ConcertServerUI/Source/ConcertServerUI/ConcertServerUI.Build.cs b/Engine/Plugins/Developer/Concert/ConcertUI/ConcertServerUI/Source/ConcertServerUI/ConcertServerUI.Build.cs
--- a/Engine/Plugins/Developer/Concert/ConcertUI/ConcertServerUI/Source/ConcertServerUI/ConcertServerUI.Build.cs
+++ b/Engine/Plugins/Developer/Concert/ConcertUI/ConcertServerUI/Source/ConcertServerUI/ConcertServerUI.Build.cs
@@ -16,6 +16,7 @@ namespace UnrealBuildTool.Rules
 					"Core",
 					"CoreUObject",
 					"Concert",
+					"ConcertSharedSlate",
 					"ConcertSyncCore",
 					"ConcertTransport",
 					"EditorStyle",			// Needed so OutputLog module to work correctly
diff --git a/Engine/Plugins/Developer/Concert/ConcertUI/ConcertServerUI/Source/ConcertServerUI/Private/ConcertServerUIModule.cpp b/Engine/Plugins/Developer/Concert/ConcertUI/ConcertServerUI/Source/ConcertServerUI/Private/ConcertServerUIModule.cpp
--- a/Engine/Plugins/Developer/Concert/ConcertUI/ConcertServerUI/Source/ConcertServerUI/Private/ConcertServerUIModule.cpp
+++ b/Engine/Plugins/Developer/Concert/ConcertUI/ConcertServerUI/Source/ConcertServerUI/Private/ConcertServerUIModule.cpp
@@ -3,6 +3,7 @@
 #include "IConcertServerUIModule.h"
 
 #include "ConcertConsoleCommandExecutor.h"
+#include "ConcertFrontendStyle.h"
 #include "ConcertServerStyle.h"
 #include "ConcertSyncServerLoopInitArgs.h"
 #include "Widgets/ConcertServerWindowController.h"
@@ -58,6 +59,7 @@ private:
 	{
 		FModuleManager::Get().LoadModuleChecked("EditorStyle");
 		FConcertServerStyle::Initialize();
+		FConcertFrontendStyle::Initialize();
 
 		// Log history must be initialized before MU server loop init prints any messages	
 		FModuleManager::Get().LoadModuleChecked("OutputLog");
@@ -73,6 +75,7 @@ private:
 		IModularFeatures::Get().RegisterModularFeature(IConsoleCommandExecutor::ModularFeatureName(), CommandExecutor.Get());
         
         WindowController = MakeShared<FConcertServerWindowController>(FConcertServerWindowInitParams{ SyncServer, MultiUserServerLayoutIni });
+		WindowController->CreateWindow();
 	}
 	
 	void TickSlate(double Tick)
diff --git a/Engine/Plugins/Developer/Concert/ConcertUI/ConcertServerUI/Source/ConcertServerUI/Private/Widgets/Browser/ConcertServerSessionBrowserController.cpp b/Engine/Plugins/Developer/Concert/ConcertUI/ConcertServerUI/Source/ConcertServerUI/Private/Widgets/Browser/ConcertServerSessionBrowserController.cpp
--- a/Engine/Plugins/Developer/Concert/ConcertUI/ConcertServerUI/Source/ConcertServerUI/Private/Widgets/Browser/ConcertServerSessionBrowserController.cpp
+++ b/Engine/Plugins/Developer/Concert/ConcertUI/ConcertServerUI/Source/ConcertServerUI/Private/Widgets/Browser/ConcertServerSessionBrowserController.cpp
@@ -3,16 +3,25 @@
 #include "ConcertServerSessionBrowserController.h"
 
 #include "ConcertServerStyle.h"
+#include "IConcertServer.h"
+#include "IConcertSyncServer.h"
 #include "Framework/Docking/TabManager.h"
 #include "SConcertServerSessionBrowser.h"
+#include "Framework/Notifications/NotificationManager.h"
+#include "Misc/AsyncTaskNotification.h"
+#include "SessionBrowser/ConcertSessionItem.h"
 #include "Textures/SlateIcon.h"
 #include "Widgets/ConcertServerTabs.h"
+#include "Widgets/ConcertServerWindowController.h"
 #include "Widgets/Docking/SDockTab.h"
+#include "Widgets/Notifications/SNotificationList.h"
 
 #define LOCTEXT_NAMESPACE "UnrealMultiUserUI"
 
 void FConcertServerSessionBrowserController::Init(const FConcertComponentInitParams& Params)
 {
+	ServerInstance = Params.Server;
+	Owner = Params.WindowController;
 	FGlobalTabmanager::Get()->RegisterTabSpawner(
 			ConcertServerTabs::GetSessionBrowserTabId(),
 			FOnSpawnTab::CreateRaw(this, &FConcertServerSessionBrowserController::SpawnSessionBrowserTab)
@@ -22,6 +31,122 @@ void FConcertServerSessionBrowserController::Init(const FConcertComponentInitPar
 		.SetIcon(FSlateIcon(FConcertServerStyle::GetStyleSetName(), TEXT("Concert.MultiUser")));
 }
 
+TArray<FConcertServerInfo> FConcertServerSessionBrowserController::GetServers() const
+{
+	return { ServerInstance->GetConcertServer()->GetServerInfo() };
+}
+
+TArray<IConcertSessionBrowserController::FActiveSessionInfo> FConcertServerSessionBrowserController::GetActiveSessions() const
+{
+	const FConcertServerInfo& ServerInfo = ServerInstance->GetConcertServer()->GetServerInfo();
+	const TArray<TSharedPtr<IConcertServerSession>> ServerSessions = ServerInstance->GetConcertServer()->GetLiveSessions();
+	
+	TArray<FActiveSessionInfo> Result;
+	Result.Reserve(ServerSessions.Num());
+	for (const TSharedPtr<IConcertServerSession>& LiveSession : ServerSessions)
+	{
+		FActiveSessionInfo Info{ ServerInfo, LiveSession->GetSessionInfo(), LiveSession->GetSessionClients() };
+		Result.Add(Info);
+	}
+
+	return Result;
+}
+
+TArray<IConcertSessionBrowserController::FArchivedSessionInfo> FConcertServerSessionBrowserController::GetArchivedSessions() const
+{
+	const FConcertServerInfo& ServerInfo = ServerInstance->GetConcertServer()->GetServerInfo();
+	const TArray<FConcertSessionInfo> ConcertSessionInfos = ServerInstance->GetConcertServer()->GetArchivedSessionInfos();
+	
+	TArray<FArchivedSessionInfo> Result;
+	Result.Reserve(ConcertSessionInfos.Num());
+	for (const FConcertSessionInfo& SessionInfo : ConcertSessionInfos)
+	{
+		Result.Emplace(ServerInfo, SessionInfo);
+	}
+
+	return Result;
+}
+
+TOptional<FConcertSessionInfo> FConcertServerSessionBrowserController::GetActiveSessionInfo(const FGuid& AdminEndpoint, const FGuid& SessionId) const
+{
+	const TSharedPtr<IConcertServerSession> LiveSession = ServerInstance->GetConcertServer()->GetLiveSession(SessionId);
+	return LiveSession ? LiveSession->GetSessionInfo() : TOptional<FConcertSessionInfo>{};
+}
+
+TOptional<FConcertSessionInfo> FConcertServerSessionBrowserController::GetArchivedSessionInfo(const FGuid& AdminEndpoint, const FGuid& SessionId) const
+{
+	return ServerInstance->GetConcertServer()->GetArchivedSessionInfo(SessionId);
+}
+
+void FConcertServerSessionBrowserController::CreateSession(const FGuid& ServerAdminEndpointId, const FString& SessionName)
+{
+	FConcertSessionInfo SessionInfo = ServerInstance->GetConcertServer()->CreateSessionInfo();
+	SessionInfo.SessionName = SessionName;
+	SessionInfo.Settings.Initialize();
+	FConcertSessionVersionInfo VersionInfo;
+	VersionInfo.Initialize();
+	SessionInfo.VersionInfos.Emplace(VersionInfo);
+	
+	FText FailureReason = FText::GetEmpty();
+	const bool bSuccess = ServerInstance->GetConcertServer()->CreateSession(SessionInfo, FailureReason) != nullptr;
+	
+	NotifyUserOfFinishedSessionAction(bSuccess,
+		bSuccess ? FText::Format(LOCTEXT("CreatedSessionFmt", "Created Session '{0}'"), FText::FromString(SessionName)) : FText::Format(LOCTEXT("FailedToCreateSessionFmt", "Failed to create Session '{0}'"), FText::FromString(SessionName)),
+		FailureReason
+		);
+}
+
+void FConcertServerSessionBrowserController::ArchiveSession(const FGuid& ServerAdminEndpointId, const FGuid& SessionId, const FString& ArchiveName, const FConcertSessionFilter& SessionFilter)
+{
+	FText FailureReason = FText::GetEmpty();
+	const bool bSuccess = ServerInstance->GetConcertServer()->ArchiveSession(SessionId, ArchiveName, SessionFilter, FailureReason).IsValid();
+	
+	NotifyUserOfFinishedSessionAction(bSuccess,
+		bSuccess ? FText::Format(LOCTEXT("ArchivedSessionFmt", "Archived Session '{0}'"),  FText::FromString(ArchiveName)) : FText::Format(LOCTEXT("FailedToArchivedSessionFmt", "Failed to archive Session '{0}'"), FText::FromString(ArchiveName)),
+		FailureReason
+		);
+}
+
+void FConcertServerSessionBrowserController::RestoreSession(const FGuid& ServerAdminEndpointId, const FGuid& SessionId, const FString& RestoredName, const FConcertSessionFilter& SessionFilter)
+{
+	if (const TOptional<FConcertSessionInfo> SessionInfo = GetArchivedSessionInfo(ServerAdminEndpointId, SessionId))
+	{
+		FText FailureReason = FText::GetEmpty();
+		const bool bSuccess = ServerInstance->GetConcertServer()->RestoreSession(SessionId, *SessionInfo, SessionFilter, FailureReason).IsValid();
+		NotifyUserOfFinishedSessionAction(bSuccess,
+			bSuccess ? FText::Format(LOCTEXT("RestoreSessionFmt", "Restored Session '{0}'"),  FText::FromString(RestoredName)) : FText::Format(LOCTEXT("FailedToRestoreSessionFmt", "Failed to restore Session '{0}'"), FText::FromString(RestoredName)),
+			FailureReason
+			);
+	}
+}
+
+void FConcertServerSessionBrowserController::RenameSession(const FGuid& ServerAdminEndpointId, const FGuid& SessionId, const FString& NewName)
+{
+	if (const TOptional<FConcertSessionInfo> SessionInfo = GetArchivedSessionInfo(ServerAdminEndpointId, SessionId))
+	{
+		FText FailureReason = FText::GetEmpty();
+		const bool bSuccess = ServerInstance->GetConcertServer()->RenameSession(SessionId, NewName, FailureReason);
+		NotifyUserOfFinishedSessionAction(bSuccess,
+				bSuccess ? FText::Format(LOCTEXT("RenameSessionFmt", "Rename Session '{0}' as '{1}'"), FText::FromString(SessionInfo->SessionName), FText::FromString(NewName))
+					: FText::Format(LOCTEXT("FailedToArchivedSessionFmt", "Failed to rename Session '{0}' as '{1}'"), FText::FromString(SessionInfo->SessionName), FText::FromString(NewName)),
+				FailureReason
+				);
+	}
+}
+
+void FConcertServerSessionBrowserController::DeleteSession(const FGuid& ServerAdminEndpointId, const FGuid& SessionId)
+{
+	if (const TOptional<FConcertSessionInfo> SessionInfo = GetActiveSessionInfo(ServerAdminEndpointId, SessionId))
+	{
+		FText FailureReason;
+		const bool bSuccess = ServerInstance->GetConcertServer()->DestroySession(SessionId, FailureReason);
+		NotifyUserOfFinishedSessionAction(bSuccess,
+			bSuccess ? FText::Format(LOCTEXT("ArchivedSessionFmt", "Deleted Session '{0}'"), FText::FromString(SessionInfo->SessionName)) : FText::Format(LOCTEXT("FailedToArchivedSessionFmt", "Failed to archive Session '{0}'"), FText::FromString(SessionInfo->SessionName)),
+			FailureReason
+			);
+	}
+}
+
 TSharedRef<SDockTab> FConcertServerSessionBrowserController::SpawnSessionBrowserTab(const FSpawnTabArgs& Args)
 {
 	const TSharedRef<SDockTab> DockTab = SNew(SDockTab)
@@ -29,10 +154,28 @@ TSharedRef<SDockTab> FConcertServerSessionBrowserController::SpawnSessionBrowser
 		.TabRole(MajorTab)
 		[
 			SAssignNew(ConcertBrowser, SConcertServerSessionBrowser, SharedThis(this))
+				.DoubleClickSession(this, &FConcertServerSessionBrowserController::OpenSession)
 		];
 
 	FGlobalTabmanager::Get()->SetMainTab(DockTab);
 	return DockTab;
 }
 
+void FConcertServerSessionBrowserController::OpenSession(TSharedPtr<FConcertSessionItem> SessionItem)
+{
+	Owner.Pin()->OpenSessionTab(SessionItem->SessionId);
+}
+
+void FConcertServerSessionBrowserController::NotifyUserOfFinishedSessionAction(const bool bSuccess, const FText& Title, const FText& Details)
+{
+	FNotificationInfo NotificationInfo(Title);
+	NotificationInfo.SubText = Details;
+	FSlateNotificationManager::Get().AddNotification(NotificationInfo);
+
+	if (bSuccess)
+	{
+		ConcertBrowser->RefreshSessionList();
+	}
+}
+
 #undef LOCTEXT_NAMESPACE 
diff --git a/Engine/Plugins/Developer/Concert/ConcertUI/ConcertServerUI/Source/ConcertServerUI/Private/Widgets/Browser/ConcertServerSessionBrowserController.h b/Engine/Plugins/Developer/Concert/ConcertUI/ConcertServerUI/Source/ConcertServerUI/Private/Widgets/Browser/ConcertServerSessionBrowserController.h
--- a/Engine/Plugins/Developer/Concert/ConcertUI/ConcertServerUI/Source/ConcertServerUI/Private/Widgets/Browser/ConcertServerSessionBrowserController.h
+++ b/Engine/Plugins/Developer/Concert/ConcertUI/ConcertServerUI/Source/ConcertServerUI/Private/Widgets/Browser/ConcertServerSessionBrowserController.h
@@ -3,9 +3,10 @@
 #pragma once
 
 #include "CoreMinimal.h"
-#include "Widgets/Browser/IConcertServerSessionBrowserController.h"
+#include "SessionBrowser/IConcertSessionBrowserController.h"
 #include "Widgets/IConcertComponent.h"
 
+class FConcertSessionItem;
 class FSpawnTabArgs;
 class SDockTab;
 class SConcertServerSessionBrowser;
@@ -15,17 +16,52 @@ class FConcertServerSessionBrowserController
 	:
 	public TSharedFromThis<FConcertServerSessionBrowserController>,
 	public IConcertComponent,
-	public IConcertServerSessionBrowserController
+	public IConcertSessionBrowserController
 {
 public:
 
 	//~ Begin IConcertComponent Interface
 	virtual void Init(const FConcertComponentInitParams& Params) override;
 	//~ End IConcertComponent Interface
+	
+	//~ Begin IConcertSessionBrowserController Interface
+	virtual TArray<FConcertServerInfo> GetServers() const override;
+	virtual TArray<FActiveSessionInfo> GetActiveSessions() const override;
+	virtual TArray<FArchivedSessionInfo> GetArchivedSessions() const override;
+	virtual TOptional<FConcertSessionInfo> GetActiveSessionInfo(const FGuid& AdminEndpoint, const FGuid& SessionId) const override;
+	virtual TOptional<FConcertSessionInfo> GetArchivedSessionInfo(const FGuid& AdminEndpoint, const FGuid& SessionId) const override;
+	
+	virtual void CreateSession(const FGuid& ServerAdminEndpointId, const FString& SessionName) override;
+	virtual void ArchiveSession(const FGuid& ServerAdminEndpointId, const FGuid& SessionId, const FString& ArchiveName, const FConcertSessionFilter& SessionFilter) override;
+	virtual void RestoreSession(const FGuid& ServerAdminEndpointId, const FGuid& SessionId, const FString& RestoredName, const FConcertSessionFilter& SessionFilter) override;
+	virtual void RenameActiveSession(const FGuid& ServerAdminEndpointId, const FGuid& SessionId, const FString& NewName) override { RenameSession(ServerAdminEndpointId, SessionId, NewName); }
+	virtual void RenameArchivedSession(const FGuid& ServerAdminEndpointId, const FGuid& SessionId, const FString& NewName) override { RenameSession(ServerAdminEndpointId, SessionId, NewName); }
+	virtual void DeleteActiveSession(const FGuid& ServerAdminEndpointId, const FGuid& SessionId) override { DeleteSession(ServerAdminEndpointId, SessionId); }
+	virtual void DeleteArchivedSession(const FGuid& ServerAdminEndpointId, const FGuid& SessionId) override { DeleteSession(ServerAdminEndpointId, SessionId); }
+	
+	// The server operator always has permission for these actions:
+	virtual bool CanRenameActiveSession(const FGuid& ServerAdminEndpointId, const FGuid& SessionId) const override { return true; }
+	virtual bool CanRenameArchivedSession(const FGuid& ServerAdminEndpointId, const FGuid& SessionId) const override { return true; }
+	virtual bool CanDeleteActiveSession(const FGuid& ServerAdminEndpointId, const FGuid& SessionId) const override { return true; }
+	virtual bool CanDeleteArchivedSession(const FGuid& ServerAdminEndpointId, const FGuid& SessionId) const override { return true; }
+	//~ Begin IConcertSessionBrowserController Interface
 
 private:
 
+	/** Server instance we're representing */
+	TSharedPtr<IConcertSyncServer> ServerInstance;
+	/** Used to open selected sessions */
+	TWeakPtr<FConcertServerWindowController> Owner;
+	
 	TSharedPtr<SConcertServerSessionBrowser> ConcertBrowser;
 	
+	void RenameSession(const FGuid& ServerAdminEndpointId, const FGuid& SessionId, const FString& NewName);
+	void DeleteSession(const FGuid& ServerAdminEndpointId, const FGuid& SessionId);
+	
 	TSharedRef<SDockTab> SpawnSessionBrowserTab(const FSpawnTabArgs& Args);
+	void OpenSession(TSharedPtr<FConcertSessionItem> SessionItem);
+	
+	void NotifyUserOfFinishedSessionAction(const bool bSuccess, const FText& Title) { NotifyUserOfFinishedSessionAction(bSuccess, Title, FText::GetEmpty()); }
+    /** Shared logic for displaying a notification after creating, archiving, restoring, renaming and deleting a session. */
+	void NotifyUserOfFinishedSessionAction(const bool bSuccess, const FText& Title, const FText& Details);
 };
diff --git a/Engine/Plugins/Developer/Concert/ConcertUI/ConcertServerUI/Source/ConcertServerUI/Private/Widgets/Browser/IConcertServerSessionBrowserController.h b/Engine/Plugins/Developer/Concert/ConcertUI/ConcertServerUI/Source/ConcertServerUI/Private/Widgets/Browser/IConcertServerSessionBrowserController.h
deleted file mode 100644
diff --git a/Engine/Plugins/Developer/Concert/ConcertUI/ConcertServerUI/Source/ConcertServerUI/Private/Widgets/Browser/SConcertServerSessionBrowser.cpp b/Engine/Plugins/Developer/Concert/ConcertUI/ConcertServerUI/Source/ConcertServerUI/Private/Widgets/Browser/SConcertServerSessionBrowser.cpp
--- a/Engine/Plugins/Developer/Concert/ConcertUI/ConcertServerUI/Source/ConcertServerUI/Private/Widgets/Browser/SConcertServerSessionBrowser.cpp
+++ b/Engine/Plugins/Developer/Concert/ConcertUI/ConcertServerUI/Source/ConcertServerUI/Private/Widgets/Browser/SConcertServerSessionBrowser.cpp
@@ -2,18 +2,22 @@
 
 #include "SConcertServerSessionBrowser.h"
 
+#include "SessionBrowser/SConcertSessionBrowser.h"
+#include "Widgets/Browser/ConcertServerSessionBrowserController.h"
+#include "Widgets/StatusBar/SConcertStatusBar.h"
+
 #include "SPositiveActionButton.h"
-#include "OutputLog/Public/OutputLogModule.h"
+#include "SessionBrowser/ConcertSessionItem.h"
 #include "Widgets/ConcertServerTabs.h"
 #include "Widgets/Input/SSearchBox.h"
 #include "Widgets/Layout/SBorder.h"
 #include "Widgets/SBoxPanel.h"
-#include "Widgets/StatusBar/SConcertStatusBar.h"
 
 #define LOCTEXT_NAMESPACE "UnrealMultiUserUI"
 
-void SConcertServerSessionBrowser::Construct(const FArguments& InArgs, TWeakPtr<IConcertServerSessionBrowserController> InController)
+void SConcertServerSessionBrowser::Construct(const FArguments& InArgs, TSharedRef<IConcertSessionBrowserController> InController)
 {
+	Controller = InController;
 	ChildSlot
 	[
 		SNew(SBorder)
@@ -22,18 +26,12 @@ void SConcertServerSessionBrowser::Construct(const FArguments& InArgs, TWeakPtr<
 		[
 			SNew(SVerticalBox)
 
-			+SVerticalBox::Slot()
-			.AutoHeight()
-			[
-				MakeControlBar()
-			]
-
 			// Session list.
 			+SVerticalBox::Slot()
 			.FillHeight(1.0f)
 			.Padding(1.0f, 2.0f)
 			[
-				MakeSessionTableView()
+				MakeSessionTableView(InArgs)
 			]
 
 			+SVerticalBox::Slot()
@@ -47,53 +45,11 @@ void SConcertServerSessionBrowser::Construct(const FArguments& InArgs, TWeakPtr<
 	];
 }
 
-TSharedRef<SWidget> SConcertServerSessionBrowser::MakeControlBar()
-{
-	return SNew(SHorizontalBox)
-
-		// + New Session
-		+SHorizontalBox::Slot()
-		.AutoWidth()
-		[
-			SNew(SPositiveActionButton)
-				.OnClicked(this, &SConcertServerSessionBrowser::OnNewSessionClicked)
-				.Text(LOCTEXT("NewSession", "New Session"))
-		]
-
-		// The search text.
-		+SHorizontalBox::Slot()
-		.FillWidth(1.0)
-		.Padding(4.0f, 3.0f, 8.0f, 3.0f)
-		[
-			SAssignNew(SearchBox, SSearchBox)
-			.HintText(LOCTEXT("SearchHint", "Search sessions"))
-			.OnTextChanged(this, &SConcertServerSessionBrowser::OnSearchTextChanged)
-			.OnTextCommitted(this, &SConcertServerSessionBrowser::OnSearchTextCommitted)
-			.DelayChangeNotificationsWhileTyping(true)
-		]
-	;
-}
-
-TSharedRef<SWidget> SConcertServerSessionBrowser::MakeSessionTableView()
-{
-	//return FOutputLogModule::Get().MakeOutputLogDrawerWidget(FSimpleDelegate());
-	return SNullWidget::NullWidget;
-}
-
-FReply SConcertServerSessionBrowser::OnNewSessionClicked()
-{
-	// TODO:
-	return FReply::Handled();
-}
-
-void SConcertServerSessionBrowser::OnSearchTextChanged(const FText& InFilterText)
-{
-	// TODO:
-}
-
-void SConcertServerSessionBrowser::OnSearchTextCommitted(const FText& InFilterText, ETextCommit::Type CommitType)
+TSharedRef<SWidget> SConcertServerSessionBrowser::MakeSessionTableView(const FArguments& InArgs)
 {
-	// TODO:
+	SearchText = MakeShared<FText>();
+	return SAssignNew(SessionBrowser, SConcertSessionBrowser, Controller.Pin().ToSharedRef(), SearchText)
+		.OnSessionDoubleClicked(InArgs._DoubleClickSession);
 }
 
 #undef LOCTEXT_NAMESPACE
diff --git a/Engine/Plugins/Developer/Concert/ConcertUI/ConcertServerUI/Source/ConcertServerUI/Private/Widgets/Browser/SConcertServerSessionBrowser.h b/Engine/Plugins/Developer/Concert/ConcertUI/ConcertServerUI/Source/ConcertServerUI/Private/Widgets/Browser/SConcertServerSessionBrowser.h
--- a/Engine/Plugins/Developer/Concert/ConcertUI/ConcertServerUI/Source/ConcertServerUI/Private/Widgets/Browser/SConcertServerSessionBrowser.h
+++ b/Engine/Plugins/Developer/Concert/ConcertUI/ConcertServerUI/Source/ConcertServerUI/Private/Widgets/Browser/SConcertServerSessionBrowser.h
@@ -3,33 +3,33 @@
 #pragma once
 
 #include "CoreMinimal.h"
+#include "SessionBrowser/SConcertSessionBrowser.h"
 #include "Widgets/DeclarativeSyntaxSupport.h"
 #include "Widgets/SCompoundWidget.h"
 
-class IConcertServerSessionBrowserController;
+class FConcertServerSessionBrowserController;
 class SSearchBox;
 
 /** Shows a list of server sessions */
 class SConcertServerSessionBrowser : public SCompoundWidget
 {
 public:
+	
 	SLATE_BEGIN_ARGS(SConcertServerSessionBrowser) { }
+		SLATE_EVENT(FSessionDelegate, DoubleClickSession)
 	SLATE_END_ARGS()
 
-	void Construct(const FArguments& InArgs, TWeakPtr<IConcertServerSessionBrowserController> InController);
+	void Construct(const FArguments& InArgs, TSharedRef<IConcertSessionBrowserController> InController);
 
+	void RefreshSessionList() { SessionBrowser->RefreshSessionList(); }
+	
 private:
 
 	/** We can ask the controller about information and notify it about UI events. */
-	TWeakPtr<IConcertServerSessionBrowserController> Controller;
-	
-	TSharedPtr<SSearchBox> SearchBox;
-	
-	TSharedRef<SWidget> MakeControlBar();
-	TSharedRef<SWidget> MakeSessionTableView();
-
-	FReply OnNewSessionClicked();
+	TWeakPtr<IConcertSessionBrowserController> Controller;
 
-	void OnSearchTextChanged(const FText& InFilterText);
-	void OnSearchTextCommitted(const FText& InFilterText, ETextCommit::Type CommitType);
+	TSharedPtr<FText> SearchText;
+	TSharedPtr<SConcertSessionBrowser> SessionBrowser;
+	
+	TSharedRef<SWidget> MakeSessionTableView(const FArguments& InArgs);
 };
diff --git a/Engine/Plugins/Developer/Concert/ConcertUI/ConcertServerUI/Source/ConcertServerUI/Private/Widgets/ConcertServerWindowController.cpp b/Engine/Plugins/Developer/Concert/ConcertUI/ConcertServerUI/Source/ConcertServerUI/Private/Widgets/ConcertServerWindowController.cpp
--- a/Engine/Plugins/Developer/Concert/ConcertUI/ConcertServerUI/Source/ConcertServerUI/Private/Widgets/ConcertServerWindowController.cpp
+++ b/Engine/Plugins/Developer/Concert/ConcertUI/ConcertServerUI/Source/ConcertServerUI/Private/Widgets/ConcertServerWindowController.cpp
@@ -3,6 +3,8 @@
 #include "ConcertServerWindowController.h"
 
 #include "ConcertServerTabs.h"
+#include "IConcertServer.h"
+#include "IConcertSyncServer.h"
 #include "Browser/ConcertServerSessionBrowserController.h"
 
 #include "Framework/Application/SlateApplication.h"
@@ -12,6 +14,7 @@
 #include "HAL/PlatformApplicationMisc.h"
 #include "Misc/ConfigCacheIni.h"
 #include "OutputLog/Public/OutputLogModule.h"
+#include "Session/ConcertServerSessionTab.h"
 #include "Widgets/SWindow.h"
 
 #define LOCTEXT_NAMESPACE "UnrealMultiUserUI"
@@ -19,8 +22,13 @@
 FConcertServerWindowController::FConcertServerWindowController(const FConcertServerWindowInitParams& Params)
 	: MultiUserServerLayoutIni(Params.MultiUserServerLayoutIni)
 {
+	ServerInstance = Params.Server;
 	SessionBrowserController = MakeShared<FConcertServerSessionBrowserController>();
-	InitComponents(Params);
+}
+
+void FConcertServerWindowController::CreateWindow()
+{
+	InitComponents();
 
 	FDisplayMetrics DisplayMetrics;
 	FSlateApplication::Get().GetDisplayMetrics(DisplayMetrics);
@@ -28,7 +36,7 @@ FConcertServerWindowController::FConcertServerWindowController(const FConcertSer
 
 	const bool bEmbedTitleAreaContent = false;
 	const FVector2D ClientSize(960.0f * DPIScaleFactor, 640.0f * DPIScaleFactor);
-	TSharedRef<SWindow> RootWindow = SNew(SWindow)
+	TSharedRef<SWindow> RootWindowRef = SNew(SWindow)
 		.Title(LOCTEXT("WindowTitle", "Unreal Multi User Server"))
 		.CreateTitleBar(!bEmbedTitleAreaContent)
 		.SupportsMaximize(true)
@@ -39,13 +47,14 @@ FConcertServerWindowController::FConcertServerWindowController(const FConcertSer
 		.AutoCenter(EAutoCenter::PreferredWorkArea)
 		.ClientSize(ClientSize)
 		.AdjustInitialSizeAndPositionForDPIScale(false);
+	RootWindow = RootWindowRef;
 		
 	const bool bShowRootWindowImmediately = false;
-	FSlateApplication::Get().AddWindow(RootWindow, bShowRootWindowImmediately);
-	FGlobalTabmanager::Get()->SetRootWindow(RootWindow);
+	FSlateApplication::Get().AddWindow(RootWindowRef, bShowRootWindowImmediately);
+	FGlobalTabmanager::Get()->SetRootWindow(RootWindowRef);
 	FGlobalTabmanager::Get()->SetAllowWindowMenuBar(true);
 
-	FSlateNotificationManager::Get().SetRootWindow(RootWindow);
+	FSlateNotificationManager::Get().SetRootWindow(RootWindowRef);
 	TSharedRef<FTabManager::FLayout> DefaultLayout = FTabManager::NewLayout("UnrealMultiUserServerLayout_v1.0");
 	DefaultLayout->AddArea
 	(
@@ -58,7 +67,7 @@ FConcertServerWindowController::FConcertServerWindowController(const FConcertSer
 		)
 	);
 	
-	PersistentLayout = FLayoutSaveRestore::LoadFromConfig(Params.MultiUserServerLayoutIni, DefaultLayout);
+	PersistentLayout = FLayoutSaveRestore::LoadFromConfig(MultiUserServerLayoutIni, DefaultLayout);
 	TSharedPtr<SWidget> Content = FGlobalTabmanager::Get()->RestoreFrom(PersistentLayout.ToSharedRef(), RootWindow, bEmbedTitleAreaContent, EOutputCanBeNullptr::Never);
 	RootWindow->SetContent(Content.ToSharedRef());
 
@@ -68,15 +77,42 @@ FConcertServerWindowController::FConcertServerWindowController(const FConcertSer
 	RootWindow->BringToFront(bForceWindowToFront);
 }
 
-void FConcertServerWindowController::InitComponents(const FConcertServerWindowInitParams& WindowInitParams) const
+void FConcertServerWindowController::OpenSessionTab(const FGuid& SessionId)
+{
+	if (TSharedPtr<FConcertServerSessionTab> SessionTab = GetOrRegisterSessionTab(SessionId))
+	{
+		SessionTab->OpenSessionTab();
+	}
+}
+
+TSharedPtr<FConcertServerSessionTab> FConcertServerWindowController::GetOrRegisterSessionTab(const FGuid& SessionId)
+{
+	if (TSharedPtr<IConcertServerSession> Session = ServerInstance->GetConcertServer()->GetLiveSession(SessionId))
+	{
+		const FConcertSessionInfo& SessionInfo = Session->GetSessionInfo();
+		if (TSharedRef<FConcertServerSessionTab>* FoundId = RegisteredSessions.Find(SessionInfo.SessionId))
+		{
+			return *FoundId;
+		}
+	
+		const TSharedRef<FConcertServerSessionTab> SessionTab = MakeShared<FConcertServerSessionTab>(Session.ToSharedRef(), RootWindow.ToSharedRef());
+		RegisteredSessions.Add(SessionInfo.SessionId, SessionTab);
+		return SessionTab;
+	}
+	
+	return nullptr;
+}
+
+void FConcertServerWindowController::InitComponents()
 {
-	const FConcertComponentInitParams Params { WindowInitParams.Server };
+	const FConcertComponentInitParams Params { ServerInstance.ToSharedRef(), SharedThis(this) };
 	SessionBrowserController->Init(Params);
 }
 
 void FConcertServerWindowController::OnWindowClosed(const TSharedRef<SWindow>& Window)
 {
 	SaveLayout();
+	RootWindow.Reset();
 }
 
 void FConcertServerWindowController::SaveLayout() const
diff --git a/Engine/Plugins/Developer/Concert/ConcertUI/ConcertServerUI/Source/ConcertServerUI/Private/Widgets/ConcertServerWindowController.h b/Engine/Plugins/Developer/Concert/ConcertUI/ConcertServerUI/Source/ConcertServerUI/Private/Widgets/ConcertServerWindowController.h
--- a/Engine/Plugins/Developer/Concert/ConcertUI/ConcertServerUI/Source/ConcertServerUI/Private/Widgets/ConcertServerWindowController.h
+++ b/Engine/Plugins/Developer/Concert/ConcertUI/ConcertServerUI/Source/ConcertServerUI/Private/Widgets/ConcertServerWindowController.h
@@ -5,6 +5,8 @@
 #include "CoreMinimal.h"
 #include "Framework/Docking/TabManager.h"
 
+class FConcertServerSessionTab;
+class IConcertServerSession;
 class IConcertSyncServer;
 class FConcertServerSessionBrowserController;
 class FOutputLogController;
@@ -28,22 +30,34 @@ struct FConcertServerWindowInitParams
 class FConcertServerWindowController : public TSharedFromThis<FConcertServerWindowController>
 {
 public:
-
+	
 	FConcertServerWindowController(const FConcertServerWindowInitParams& Params);
+	void CreateWindow();
 
+	/** Opens or draws attention to the tab for the given live session ID */
+	void OpenSessionTab(const FGuid& SessionId);
+	
 private:
 
 	/** The ini file to use for saving the layout */
 	FString MultiUserServerLayoutIni;
-	
 	/** Holds the current layout for saving later. */
 	TSharedPtr<FTabManager::FLayout> PersistentLayout;
 
+	TSharedPtr<IConcertSyncServer> ServerInstance;
+
+	/** The main window being managed */
+	TSharedPtr<SWindow> RootWindow;
+	TMap<FGuid, TSharedRef<FConcertServerSessionTab>> RegisteredSessions;
+	
 	/** Manages the session browser */
 	TSharedPtr<FConcertServerSessionBrowserController> SessionBrowserController;
 	
-	void InitComponents(const FConcertServerWindowInitParams& WindowInitParams) const;
+	void InitComponents();
+
+	/** Gets the manager for a session tab if the session ID is valid */
+	TSharedPtr<FConcertServerSessionTab> GetOrRegisterSessionTab(const FGuid& SessionId);
 
 	void OnWindowClosed(const TSharedRef<SWindow>& Window);
 	void SaveLayout() const;
-};
\ No newline at end of file
+};
diff --git a/Engine/Plugins/Developer/Concert/ConcertUI/ConcertServerUI/Source/ConcertServerUI/Private/Widgets/IConcertComponent.h b/Engine/Plugins/Developer/Concert/ConcertUI/ConcertServerUI/Source/ConcertServerUI/Private/Widgets/IConcertComponent.h
--- a/Engine/Plugins/Developer/Concert/ConcertUI/ConcertServerUI/Source/ConcertServerUI/Private/Widgets/IConcertComponent.h
+++ b/Engine/Plugins/Developer/Concert/ConcertUI/ConcertServerUI/Source/ConcertServerUI/Private/Widgets/IConcertComponent.h
@@ -4,11 +4,21 @@
 
 #include "CoreMinimal.h"
 
+class FConcertServerWindowController;
 class IConcertSyncServer;
 
 struct FConcertComponentInitParams
 {
+	/** The server being managed */
 	TSharedRef<IConcertSyncServer> Server;
+
+	/** Manages the server slate application */
+	TSharedRef<FConcertServerWindowController> WindowController;
+
+	FConcertComponentInitParams(TSharedRef<IConcertSyncServer> Server, TSharedRef<FConcertServerWindowController> WindowController)
+		: Server(MoveTemp(Server))
+		, WindowController(MoveTemp(WindowController))
+	{}
 };
 
 /** Provides the base interface for elements in the concert server UI. */
diff --git a/Engine/Plugins/Developer/Concert/ConcertUI/ConcertServerUI/Source/ConcertServerUI/Private/Widgets/Session/ConcertServerSessionTab.cpp b/Engine/Plugins/Developer/Concert/ConcertUI/ConcertServerUI/Source/ConcertServerUI/Private/Widgets/Session/ConcertServerSessionTab.cpp
new file mode 100644
--- /dev/null
+++ b/Engine/Plugins/Developer/Concert/ConcertUI/ConcertServerUI/Source/ConcertServerUI/Private/Widgets/Session/ConcertServerSessionTab.cpp
@@ -0,0 +1,49 @@
+﻿// Copyright Epic Games, Inc. All Rights Reserved.
+
+#include "ConcertServerSessionTab.h"
+
+#include "IConcertSession.h"
+#include "Widgets/ConcertServerTabs.h"
+#include "Widgets/StatusBar/SConcertStatusBar.h"
+#include "SessionInspector/SConcertSessionInspector.h"
+
+FConcertServerSessionTab::FConcertServerSessionTab(TSharedRef<IConcertServerSession> InspectedSession, const TSharedRef<SWindow>& ConstructUnderWindow)
+	: InspectedSession(MoveTemp(InspectedSession))
+	, DockTab(CreateTab(InspectedSession, ConstructUnderWindow))
+{}
+
+void FConcertServerSessionTab::OpenSessionTab() const
+{
+	const TSharedRef<FGlobalTabmanager>& TabManager = FGlobalTabmanager::Get();
+	const FTabId TabId { *GetTabPlayerHolderId(InspectedSession) };
+	if (TabManager->FindExistingLiveTab(TabId))
+	{
+		TabManager->DrawAttention(DockTab);
+	}
+	else
+	{
+		const FTabManager::FLastMajorOrNomadTab Search(ConcertServerTabs::GetSessionBrowserTabId());
+		FGlobalTabmanager::Get()->InsertNewDocumentTab(*GetTabPlayerHolderId(InspectedSession), Search, DockTab);
+	}
+}
+
+TSharedRef<SDockTab> FConcertServerSessionTab::CreateTab(const TSharedRef<IConcertServerSession>& InspectedSession, const TSharedRef<SWindow>& ConstructUnderWindow)
+{
+	const FText Title = FText::FromString(InspectedSession->GetSessionInfo().SessionName);
+	TSharedRef<SDockTab> DockTab = SNew(SDockTab)
+		.Label(Title)
+		.TabRole(MajorTab);
+	DockTab->SetContent(
+		SNew(SConcertSessionInspector, DockTab, ConstructUnderWindow)
+			.StatusBar()
+			[
+				SNew(SConcertStatusBar, *GetTabPlayerHolderId(InspectedSession))
+			]
+		);
+	return DockTab;
+}
+
+FString FConcertServerSessionTab::GetTabPlayerHolderId(const TSharedRef<IConcertServerSession>& InspectedSession)
+{
+	return *InspectedSession->GetSessionInfo().SessionId.ToString();
+}
diff --git a/Engine/Plugins/Developer/Concert/ConcertUI/ConcertServerUI/Source/ConcertServerUI/Private/Widgets/Session/ConcertServerSessionTab.h b/Engine/Plugins/Developer/Concert/ConcertUI/ConcertServerUI/Source/ConcertServerUI/Private/Widgets/Session/ConcertServerSessionTab.h
new file mode 100644
--- /dev/null
+++ b/Engine/Plugins/Developer/Concert/ConcertUI/ConcertServerUI/Source/ConcertServerUI/Private/Widgets/Session/ConcertServerSessionTab.h
@@ -0,0 +1,34 @@
+﻿// Copyright Epic Games, Inc. All Rights Reserved.
+
+#pragma once
+
+#include "CoreMinimal.h"
+
+class IConcertServerSession;
+class SDockTab;
+class SWindow;
+
+/** Manages the tab that contains the UI for a session */
+class FConcertServerSessionTab
+{
+public:
+
+	FConcertServerSessionTab(TSharedRef<IConcertServerSession> InspectedSession, const TSharedRef<SWindow>& ConstructUnderWindow);
+
+	/** Opens or draws attention to the given tab */
+	void OpenSessionTab() const;
+
+private:
+
+	/** The session being inspected */
+	const TSharedRef<IConcertServerSession> InspectedSession;
+	
+	/** The tab containing the UI for InspectedSession */
+	const TSharedRef<SDockTab> DockTab;
+
+	/** Creates a tab widget */
+	static TSharedRef<SDockTab> CreateTab(const TSharedRef<IConcertServerSession>& InspectedSession, const TSharedRef<SWindow>& ConstructUnderWindow);
+	
+	/** Generates a tab ID for FTabManager::InsertNewDocumentTab */
+	static FString GetTabPlayerHolderId(const TSharedRef<IConcertServerSession>& InspectedSession);
+};
diff --git a/Engine/Plugins/Developer/Concert/ConcertUI/ConcertSharedSlate/Source/ConcertSharedSlate/Private/SessionBrowser/SConcertSessionBrowser.cpp b/Engine/Plugins/Developer/Concert/ConcertUI/ConcertSharedSlate/Source/ConcertSharedSlate/Private/SessionBrowser/SConcertSessionBrowser.cpp
--- a/Engine/Plugins/Developer/Concert/ConcertUI/ConcertSharedSlate/Source/ConcertSharedSlate/Private/SessionBrowser/SConcertSessionBrowser.cpp
+++ b/Engine/Plugins/Developer/Concert/ConcertUI/ConcertSharedSlate/Source/ConcertSharedSlate/Private/SessionBrowser/SConcertSessionBrowser.cpp
@@ -138,9 +138,9 @@ void SConcertSessionBrowser::RefreshSessionList()
 	Sessions.Reset();
 
 	// Populate the live sessions.
-	for (const TSharedPtr<IConcertSessionBrowserController::FActiveSessionInfo>& ActiveSession : GetController()->GetActiveSessions())
+	for (const IConcertSessionBrowserController::FActiveSessionInfo& ActiveSession : GetController()->GetActiveSessions())
 	{
-		FConcertSessionItem NewItem(FConcertSessionItem::EType::ActiveSession, ActiveSession->SessionInfo.SessionName, ActiveSession->SessionInfo.SessionId, ActiveSession->ServerInfo.ServerName, ActiveSession->ServerInfo.AdminEndpointId, ActiveSession->ServerInfo.ServerFlags);
+		FConcertSessionItem NewItem(FConcertSessionItem::EType::ActiveSession, ActiveSession.SessionInfo.SessionName, ActiveSession.SessionInfo.SessionId, ActiveSession.ServerInfo.ServerName, ActiveSession.ServerInfo.AdminEndpointId, ActiveSession.ServerInfo.ServerFlags);
 		if (!IsFilteredOut(NewItem))
 		{
 			Sessions.Emplace(MakeShared<FConcertSessionItem>(MoveTemp(NewItem)));
@@ -149,9 +149,9 @@ void SConcertSessionBrowser::RefreshSessionList()
 	}
 
 	// Populate the archived.
-	for (const TSharedPtr<IConcertSessionBrowserController::FArchivedSessionInfo>& ArchivedSession : GetController()->GetArchivedSessions())
+	for (const IConcertSessionBrowserController::FArchivedSessionInfo& ArchivedSession : GetController()->GetArchivedSessions())
 	{
-		FConcertSessionItem NewItem(FConcertSessionItem::EType::ArchivedSession, ArchivedSession->SessionInfo.SessionName, ArchivedSession->SessionInfo.SessionId, ArchivedSession->ServerInfo.ServerName, ArchivedSession->ServerInfo.AdminEndpointId, ArchivedSession->ServerInfo.ServerFlags);
+		FConcertSessionItem NewItem(FConcertSessionItem::EType::ArchivedSession, ArchivedSession.SessionInfo.SessionName, ArchivedSession.SessionInfo.SessionId, ArchivedSession.ServerInfo.ServerName, ArchivedSession.ServerInfo.AdminEndpointId, ArchivedSession.ServerInfo.ServerFlags);
 		if (!IsFilteredOut(NewItem))
 		{
 			Sessions.Emplace(MakeShared<FConcertSessionItem>(MoveTemp(NewItem)));
@@ -486,7 +486,7 @@ TSharedRef<ITableRow> SConcertSessionBrowser::OnGenerateSessionRowWidget(TShared
 
 TSharedRef<ITableRow> SConcertSessionBrowser::MakeActiveSessionRowWidget(const TSharedPtr<FConcertSessionItem>& ActiveItem, const TSharedRef<STableViewBase>& OwnerTable)
 {
-	const FConcertSessionInfo* SessionInfo = GetController()->GetActiveSessionInfo(ActiveItem->ServerAdminEndpointId, ActiveItem->SessionId);
+	const TOptional<FConcertSessionInfo> SessionInfo = GetController()->GetActiveSessionInfo(ActiveItem->ServerAdminEndpointId, ActiveItem->SessionId);
 
 	// Add an 'Active Session' row. Clicking the row icon joins the session.
 	return SNew(SSessionRow, ActiveItem, OwnerTable)
@@ -505,7 +505,7 @@ TSharedRef<ITableRow> SConcertSessionBrowser::MakeActiveSessionRowWidget(const T
 
 TSharedRef<ITableRow> SConcertSessionBrowser::MakeArchivedSessionRowWidget(const TSharedPtr<FConcertSessionItem>& ArchivedItem, const TSharedRef<STableViewBase>& OwnerTable)
 {
-	const FConcertSessionInfo* SessionInfo = GetController()->GetArchivedSessionInfo(ArchivedItem->ServerAdminEndpointId, ArchivedItem->SessionId);
+	const TOptional<FConcertSessionInfo> SessionInfo = GetController()->GetArchivedSessionInfo(ArchivedItem->ServerAdminEndpointId, ArchivedItem->SessionId);
 
 	// Add an 'Archived Session' row. Clicking the row icon adds a 'Restore as' row to the table.
 	return SNew(SSessionRow, ArchivedItem, OwnerTable)
diff --git a/Engine/Plugins/Developer/Concert/ConcertUI/ConcertSharedSlate/Source/ConcertSharedSlate/Private/SessionInspector/SConcertSessionInspector.cpp b/Engine/Plugins/Developer/Concert/ConcertUI/ConcertSharedSlate/Source/ConcertSharedSlate/Private/SessionInspector/SConcertSessionInspector.cpp
new file mode 100644
--- /dev/null
+++ b/Engine/Plugins/Developer/Concert/ConcertUI/ConcertSharedSlate/Source/ConcertSharedSlate/Private/SessionInspector/SConcertSessionInspector.cpp
@@ -0,0 +1,132 @@
+﻿// Copyright Epic Games, Inc. All Rights Reserved.
+
+#include "SessionInspector/SConcertSessionInspector.h"
+
+#include "Framework/Docking/TabManager.h"
+#include "Widgets/Docking/SDockTab.h"
+#include "Widgets/SWindow.h"
+
+#define LOCTEXT_NAMESPACE "SConcertSessionInspector"
+
+const FName SConcertSessionInspector::HistoryTabId("HistoryTabId");
+const FName SConcertSessionInspector::SessionContentTabId("SessionContentTabId");
+const FName SConcertSessionInspector::ConnectionMonitorTabId("ConnectionMonitorTabId");
+
+void SConcertSessionInspector::Construct(const FArguments& InArgs, const TSharedRef<SDockTab>& ConstructUnderMajorTab, const TSharedRef<SWindow>& ConstructUnderWindow)
+{
+	ChildSlot
+	[
+		SNew(SBorder)
+		.BorderImage(FAppStyle::Get().GetBrush("ToolPanel.GroupBorder"))
+		.Padding(FMargin(1.0f, 2.0f))
+		[
+			SNew(SVerticalBox)
+
+			// Toolbar
+			+SVerticalBox::Slot()
+			.AutoHeight()
+			[
+				SNullWidget::NullWidget
+			]
+
+			// Content
+			+SVerticalBox::Slot()
+			.FillHeight(1.f)
+			[
+				SNew(SBorder)
+				.BorderImage(FAppStyle::Get().GetBrush("Brushes.Title"))
+				.Padding(FMargin(0.f, 0.f, 0.f, 5.0f)) // Visually separate status bar from tabs
+				[
+					CreateTabs(ConstructUnderMajorTab, ConstructUnderWindow)
+				]
+			]
+
+			// Status bar
+			+SVerticalBox::Slot()
+			.AutoHeight()
+			.HAlign(HAlign_Fill)
+			[
+				InArgs._StatusBar.Widget
+			]
+		]
+	];
+}
+
+TSharedRef<SWidget> SConcertSessionInspector::CreateTabs(const TSharedRef<SDockTab>& ConstructUnderMajorTab, const TSharedRef<SWindow>& ConstructUnderWindow)
+{
+	TabManager = FGlobalTabmanager::Get()->NewTabManager(ConstructUnderMajorTab);
+	
+	TabManager->RegisterTabSpawner(HistoryTabId, FOnSpawnTab::CreateSP(this, &SConcertSessionInspector::SpawnActivityHistory))
+		.SetDisplayName(LOCTEXT("ActivityHistoryLabel", "History"));
+	
+	TabManager->RegisterTabSpawner(SessionContentTabId, FOnSpawnTab::CreateSP(this, &SConcertSessionInspector::SpawnSessionContent))
+		.SetDisplayName(LOCTEXT("SessionContentLabel", "Session Content"));
+	
+	TabManager->RegisterTabSpawner(ConnectionMonitorTabId, FOnSpawnTab::CreateSP(this, &SConcertSessionInspector::SpawnConnectionMonitor))
+		.SetDisplayName(LOCTEXT("ConnectionMonitorLabel", "Connection Monitor"));
+
+	const TSharedRef<FTabManager::FLayout> Layout = FTabManager::NewLayout("ConcertSessionLayout_v0.3")
+		->AddArea
+		(
+			FTabManager::NewPrimaryArea()
+				->SetOrientation(Orient_Horizontal)
+				->Split
+				(
+					FTabManager::NewStack()
+					->SetSizeCoefficient(0.5f)
+					->AddTab(HistoryTabId, ETabState::OpenedTab)
+				)
+				->Split
+				(
+					FTabManager::NewSplitter()
+					->SetSizeCoefficient(0.5f)
+					->SetOrientation(Orient_Vertical)
+					->Split
+					(
+						FTabManager::NewStack()
+						->SetSizeCoefficient(0.5f)
+						->AddTab(SessionContentTabId, ETabState::OpenedTab)
+					)
+					->Split
+					(
+						FTabManager::NewStack()
+						->SetSizeCoefficient(0.5f)
+						->AddTab(ConnectionMonitorTabId, ETabState::OpenedTab)
+					)
+				)
+		);
+
+	return TabManager->RestoreFrom(Layout, ConstructUnderWindow).ToSharedRef();
+}
+
+TSharedRef<SDockTab> SConcertSessionInspector::SpawnActivityHistory(const FSpawnTabArgs& Args)
+{
+	return SNew(SDockTab)
+		.Label(LOCTEXT("ActivityHistoryLabel", "History"))
+		.TabRole(PanelTab)
+		[
+			SNullWidget::NullWidget
+		];
+}
+
+TSharedRef<SDockTab> SConcertSessionInspector::SpawnSessionContent(const FSpawnTabArgs& Args)
+{
+	return SNew(SDockTab)
+		.Label(LOCTEXT("SessionContentLabel", "Session Content"))
+		.TabRole(PanelTab)
+		[
+			SNullWidget::NullWidget
+		];
+}
+
+TSharedRef<SDockTab> SConcertSessionInspector::SpawnConnectionMonitor(const FSpawnTabArgs& Args)
+{
+	return SNew(SDockTab)
+		.Label(LOCTEXT("ConnectionMonitorLabel", "Connection Monitor"))
+		.TabRole(PanelTab)
+		[
+			SNullWidget::NullWidget
+		];
+}
+
+#undef LOCTEXT_NAMESPACE
diff --git a/Engine/Plugins/Developer/Concert/ConcertUI/ConcertSharedSlate/Source/ConcertSharedSlate/Public/SessionBrowser/IConcertSessionBrowserController.h b/Engine/Plugins/Developer/Concert/ConcertUI/ConcertSharedSlate/Source/ConcertSharedSlate/Public/SessionBrowser/IConcertSessionBrowserController.h
--- a/Engine/Plugins/Developer/Concert/ConcertUI/ConcertSharedSlate/Source/ConcertSharedSlate/Public/SessionBrowser/IConcertSessionBrowserController.h
+++ b/Engine/Plugins/Developer/Concert/ConcertUI/ConcertSharedSlate/Source/ConcertSharedSlate/Public/SessionBrowser/IConcertSessionBrowserController.h
@@ -20,9 +20,10 @@ public:
 		FConcertSessionInfo SessionInfo;
 		TArray<FConcertSessionClientInfo> Clients;
 
-		FActiveSessionInfo(FConcertServerInfo ServerInfo, FConcertSessionInfo SessionInfo)
+		FActiveSessionInfo(FConcertServerInfo ServerInfo, FConcertSessionInfo SessionInfo, TArray<FConcertSessionClientInfo> Clients = {})
 			: ServerInfo(MoveTemp(ServerInfo))
 			, SessionInfo(MoveTemp(SessionInfo))
+			, Clients(MoveTemp(Clients))
 		{}
 	};
 
@@ -43,21 +44,20 @@ public:
 	virtual TArray<FConcertServerInfo> GetServers() const = 0;
 
 	/** Returns the latest list of active sessions known to this controller. */
-	virtual TArray<TSharedPtr<FActiveSessionInfo>> GetActiveSessions() const = 0;
+	virtual TArray<FActiveSessionInfo> GetActiveSessions() const = 0;
 
 	/** Returns the latest list of archived sessions known to this controller. */
-	virtual TArray<TSharedPtr<FArchivedSessionInfo>> GetArchivedSessions() const = 0;
+	virtual TArray<FArchivedSessionInfo> GetArchivedSessions() const = 0;
 
 	/** Returns the active sessions info corresponding to the specified parameters. Used to display the sessions details. */
-	virtual const FConcertSessionInfo* GetActiveSessionInfo(const FGuid& AdminEndpoint, const FGuid& SessionId) const = 0;
+	virtual TOptional<FConcertSessionInfo> GetActiveSessionInfo(const FGuid& AdminEndpoint, const FGuid& SessionId) const = 0;
 
 	/** Returns the archived sessions info corresponding to the specified parameters. Used to display the sessions details. */
-	virtual const FConcertSessionInfo* GetArchivedSessionInfo(const FGuid& AdminEndpoint, const FGuid& SessionId) const = 0;
+	virtual TOptional<FConcertSessionInfo> GetArchivedSessionInfo(const FGuid& AdminEndpoint, const FGuid& SessionId) const = 0;
 
 	virtual void CreateSession(const FGuid& ServerAdminEndpointId, const FString& SessionName) = 0;
 	virtual void ArchiveSession(const FGuid& ServerAdminEndpointId, const FGuid& SessionId, const FString& ArchiveName, const FConcertSessionFilter& SessionFilter) = 0;
 	virtual void RestoreSession(const FGuid& ServerAdminEndpointId, const FGuid& SessionId, const FString& RestoredName, const FConcertSessionFilter& SessionFilter) = 0;
-	virtual void JoinSession(const FGuid& ServerAdminEndpointId, const FGuid& SessionId) = 0;
 	virtual void RenameActiveSession(const FGuid& ServerAdminEndpointId, const FGuid& SessionId, const FString& NewName) = 0;
 	virtual void RenameArchivedSession(const FGuid& ServerAdminEndpointId, const FGuid& SessionId, const FString& NewName) = 0;
 	virtual bool CanRenameActiveSession(const FGuid& ServerAdminEndpointId, const FGuid& SessionId) const = 0;
diff --git a/Engine/Plugins/Developer/Concert/ConcertUI/ConcertSharedSlate/Source/ConcertSharedSlate/Public/SessionInspector/SConcertSessionInspector.h b/Engine/Plugins/Developer/Concert/ConcertUI/ConcertSharedSlate/Source/ConcertSharedSlate/Public/SessionInspector/SConcertSessionInspector.h
new file mode 100644
--- /dev/null
+++ b/Engine/Plugins/Developer/Concert/ConcertUI/ConcertSharedSlate/Source/ConcertSharedSlate/Public/SessionInspector/SConcertSessionInspector.h
@@ -0,0 +1,44 @@
+﻿// Copyright Epic Games, Inc. All Rights Reserved.
+
+#pragma once
+
+#include "CoreMinimal.h"
+#include "Widgets/SCompoundWidget.h"
+#include "Widgets/DeclarativeSyntaxSupport.h"
+
+class FTabManager;
+class SDockTab;
+class SWindow;
+
+/**
+ * Designed to be the content of a tab showing:
+ *  - activity history (transactions stored on the server as well as who made those transactions)
+ *  - session content (list of session data saved during a Multi-user session)
+ *  - connection monitor (details about the connected clients on the given session and network info)
+ * Implements view in model-view-controller pattern.
+ */
+class CONCERTSHAREDSLATE_API SConcertSessionInspector : public SCompoundWidget
+{
+public:
+
+	static const FName HistoryTabId;
+	static const FName SessionContentTabId;
+	static const FName ConnectionMonitorTabId;
+
+	SLATE_BEGIN_ARGS(SConcertSessionInspector) {}
+		SLATE_NAMED_SLOT(FArguments, StatusBar)
+	SLATE_END_ARGS()
+
+	void Construct(const FArguments& InArgs, const TSharedRef<SDockTab>& ConstructUnderMajorTab, const TSharedRef<SWindow>& ConstructUnderWindow);
+
+private:
+	
+	/** Holds the tab manager that manages the front-end's tabs. */
+	TSharedPtr<FTabManager> TabManager;
+
+	// Spawn tabs
+	TSharedRef<SWidget> CreateTabs(const TSharedRef<SDockTab>& ConstructUnderMajorTab, const TSharedRef<SWindow>& ConstructUnderWindow);
+	TSharedRef<SDockTab> SpawnActivityHistory(const FSpawnTabArgs& Args);
+	TSharedRef<SDockTab> SpawnSessionContent(const FSpawnTabArgs& Args);
+	TSharedRef<SDockTab> SpawnConnectionMonitor(const FSpawnTabArgs& Args);
+};
diff --git a/Engine/Plugins/Developer/Concert/MultiUserClient/Source/MultiUserClient/Private/Widgets/SConcertBrowser.cpp b/Engine/Plugins/Developer/Concert/MultiUserClient/Source/MultiUserClient/Private/Widgets/SConcertBrowser.cpp
--- a/Engine/Plugins/Developer/Concert/MultiUserClient/Source/MultiUserClient/Private/Widgets/SConcertBrowser.cpp
+++ b/Engine/Plugins/Developer/Concert/MultiUserClient/Source/MultiUserClient/Private/Widgets/SConcertBrowser.cpp
@@ -117,22 +117,22 @@ public:
 	//~ Begin IConcertSessionBrowserController Interface
 	virtual TArray<FConcertServerInfo> GetServers() const override { return Servers; }
 	
-	virtual TArray<TSharedPtr<FActiveSessionInfo>> GetActiveSessions() const override
+	virtual TArray<FActiveSessionInfo> GetActiveSessions() const override
 	{
-		TArray<TSharedPtr<FActiveSessionInfo>> Result;
-		Algo::Transform(ActiveSessions, Result, [](TSharedPtr<FClientActiveSessionInfo> Info) { return Info; });
+		TArray<FActiveSessionInfo> Result;
+		Algo::Transform(ActiveSessions, Result, [](TSharedPtr<FClientActiveSessionInfo> Info) { return *Info; });
 		return Result;
 	}
 	
-	virtual TArray<TSharedPtr<FArchivedSessionInfo>> GetArchivedSessions() const override
+	virtual TArray<FArchivedSessionInfo> GetArchivedSessions() const override
 	{
-		TArray<TSharedPtr<FArchivedSessionInfo>> Result;
-		Algo::Transform(ArchivedSessions, Result, [](TSharedPtr<FClientArchivedSessionInfo> Info) { return Info; });
+		TArray<FArchivedSessionInfo> Result;
+		Algo::Transform(ArchivedSessions, Result, [](TSharedPtr<FClientArchivedSessionInfo> Info) { return *Info; });
 		return Result;
 	}
 	
-	virtual const FConcertSessionInfo* GetActiveSessionInfo(const FGuid& AdminEndpoint, const FGuid& SessionId) const override;
-	virtual const FConcertSessionInfo* GetArchivedSessionInfo(const FGuid& AdminEndpoint, const FGuid& SessionId) const override;
+	virtual TOptional<FConcertSessionInfo> GetActiveSessionInfo(const FGuid& AdminEndpoint, const FGuid& SessionId) const override;
+	virtual TOptional<FConcertSessionInfo> GetArchivedSessionInfo(const FGuid& AdminEndpoint, const FGuid& SessionId) const override;
 
 	virtual void CreateSession(const FGuid& ServerAdminEndpointId, const FString& SessionName) override;
 	virtual void ArchiveSession(const FGuid& ServerAdminEndpointId, const FGuid& SessionId, const FString& ArchiveName, const FConcertSessionFilter& SessionFilter) override;
@@ -275,24 +275,24 @@ void FConcertClientSessionBrowserController::JoinSession(const FGuid& ServerAdmi
 	ConcertClient->JoinSession(ServerAdminEndpointId, SessionId);
 }
 
-const FConcertSessionInfo* FConcertClientSessionBrowserController::GetActiveSessionInfo(const FGuid& AdminEndpoint, const FGuid& SessionId) const
+TOptional<FConcertSessionInfo> FConcertClientSessionBrowserController::GetActiveSessionInfo(const FGuid& AdminEndpoint, const FGuid& SessionId) const
 {
 	const TSharedPtr<FClientActiveSessionInfo>* SessionInfo = ActiveSessions.FindByPredicate([&AdminEndpoint, &SessionId](const TSharedPtr<FActiveSessionInfo>& MatchCandidate)
 	{
 		return MatchCandidate->ServerInfo.AdminEndpointId == AdminEndpoint && MatchCandidate->SessionInfo.SessionId == SessionId;
 	});
 
-	return SessionInfo != nullptr ? &(*SessionInfo)->SessionInfo : nullptr;
+	return SessionInfo != nullptr ? (*SessionInfo)->SessionInfo : TOptional<FConcertSessionInfo>{};
 }
 
-const FConcertSessionInfo* FConcertClientSessionBrowserController::GetArchivedSessionInfo(const FGuid& AdminEndpoint, const FGuid& SessionId) const
+TOptional<FConcertSessionInfo> FConcertClientSessionBrowserController::GetArchivedSessionInfo(const FGuid& AdminEndpoint, const FGuid& SessionId) const
 {
 	const TSharedPtr<FClientArchivedSessionInfo>* SessionInfo = ArchivedSessions.FindByPredicate([&AdminEndpoint, &SessionId](const TSharedPtr<FArchivedSessionInfo>& MatchCandidate)
 	{
 		return MatchCandidate->ServerInfo.AdminEndpointId == AdminEndpoint && MatchCandidate->SessionInfo.SessionId == SessionId;
 	});
 
-	return SessionInfo != nullptr ? &(*SessionInfo)->SessionInfo : nullptr;
+	return SessionInfo != nullptr ? (*SessionInfo)->SessionInfo : TOptional<FConcertSessionInfo>{};
 }
 
 void FConcertClientSessionBrowserController::CreateSession(const FGuid& ServerAdminEndpointId, const FString& SessionName)
@@ -334,7 +334,7 @@ void FConcertClientSessionBrowserController::ArchiveSession(const FGuid& ServerA
 void FConcertClientSessionBrowserController::RestoreSession(const FGuid& ServerAdminEndpointId, const FGuid& SessionId, const FString& RestoredName, const FConcertSessionFilter& SessionFilter)
 {
 	FString ArchivedSessionName;
-	if (const FConcertSessionInfo* SessionInfo = GetArchivedSessionInfo(ServerAdminEndpointId, SessionId))
+	if (const TOptional<FConcertSessionInfo> SessionInfo = GetArchivedSessionInfo(ServerAdminEndpointId, SessionId))
 	{
 		ArchivedSessionName = SessionInfo->SessionName;
 	}
@@ -1463,7 +1463,7 @@ TSharedRef<SWidget> SConcertClientSessionBrowser::MakeSessionDetails(TSharedPtr<
 
 TSharedRef<SWidget> SConcertClientSessionBrowser::MakeActiveSessionDetails(TSharedPtr<FConcertSessionItem> Item)
 {
-	const FConcertSessionInfo* SessionInfo = Controller->GetActiveSessionInfo(Item->ServerAdminEndpointId, Item->SessionId);
+	const TOptional<FConcertSessionInfo> SessionInfo = Controller->GetActiveSessionInfo(Item->ServerAdminEndpointId, Item->SessionId);
 	if (!SessionInfo)
 	{
 		return NoSessionSelectedPanel.ToSharedRef();
@@ -1576,7 +1576,7 @@ TSharedRef<SWidget> SConcertClientSessionBrowser::MakeActiveSessionDetails(TShar
 
 TSharedRef<SWidget> SConcertClientSessionBrowser::MakeArchivedSessionDetails(TSharedPtr<FConcertSessionItem> Item)
 {
-	const FConcertSessionInfo* SessionInfo = Controller->GetArchivedSessionInfo(Item->ServerAdminEndpointId, Item->SessionId);
+	const TOptional<FConcertSessionInfo> SessionInfo = Controller->GetArchivedSessionInfo(Item->ServerAdminEndpointId, Item->SessionId);
 	if (!SessionInfo)
 	{
 		return NoSessionSelectedPanel.ToSharedRef();
