Fix non-unity build errors.

diff --git a/Engine/Plugins/Developer/Concert/ConcertUI/ConcertServerUI/Source/ConcertServerUI/Private/Widgets/Session/ConcertServerSessionTab.cpp b/Engine/Plugins/Developer/Concert/ConcertUI/ConcertServerUI/Source/ConcertServerUI/Private/Widgets/Session/ConcertServerSessionTab.cpp
--- a/Engine/Plugins/Developer/Concert/ConcertUI/ConcertServerUI/Source/ConcertServerUI/Private/Widgets/Session/ConcertServerSessionTab.cpp
+++ b/Engine/Plugins/Developer/Concert/ConcertUI/ConcertServerUI/Source/ConcertServerUI/Private/Widgets/Session/ConcertServerSessionTab.cpp
@@ -2,8 +2,10 @@
 
 #include "ConcertServerSessionTab.h"
 
+#include "Framework/Docking/TabManager.h"
 #include "IConcertSession.h"
 #include "Widgets/ConcertServerTabs.h"
+#include "Widgets/Docking/SDockTab.h"
 #include "Widgets/StatusBar/SConcertStatusBar.h"
 #include "Widgets/Session/SConcertSessionInspector.h"
 
diff --git a/Engine/Plugins/Developer/Concert/ConcertUI/ConcertSharedSlate/Source/ConcertSharedSlate/Private/SessionBrowser/SSessionRow.h b/Engine/Plugins/Developer/Concert/ConcertUI/ConcertSharedSlate/Source/ConcertSharedSlate/Private/SessionBrowser/SSessionRow.h
--- a/Engine/Plugins/Developer/Concert/ConcertUI/ConcertSharedSlate/Source/ConcertSharedSlate/Private/SessionBrowser/SSessionRow.h
+++ b/Engine/Plugins/Developer/Concert/ConcertUI/ConcertSharedSlate/Source/ConcertSharedSlate/Private/SessionBrowser/SSessionRow.h
@@ -7,6 +7,7 @@
 #include "Widgets/DeclarativeSyntaxSupport.h"
 #include "Widgets/Text/SInlineEditableTextBlock.h"
 #include "Widgets/Views/STableRow.h"
+#include "Widgets/Views/SListView.h"
 
 class FConcertSessionItem;
 class SWidget;
diff --git a/Engine/Plugins/Developer/Concert/MultiUserClient/Source/MultiUserClient/Private/Widgets/ClientSessionHistoryController.h b/Engine/Plugins/Developer/Concert/MultiUserClient/Source/MultiUserClient/Private/Widgets/ClientSessionHistoryController.h
--- a/Engine/Plugins/Developer/Concert/MultiUserClient/Source/MultiUserClient/Private/Widgets/ClientSessionHistoryController.h
+++ b/Engine/Plugins/Developer/Concert/MultiUserClient/Source/MultiUserClient/Private/Widgets/ClientSessionHistoryController.h
@@ -9,6 +9,7 @@ struct FConcertSyncActivity;
 struct FConcertClientInfo;
 class IConcertClientWorkspace;
 class IConcertSyncClient;
+class FStructOnScope;
 
 /** Manages SSessionHistory for IConcertClients */
 class FClientSessionHistoryController : public FAbstractSessionHistoryController
