Extend session browser

diff --git a/Engine/Plugins/Developer/Concert/ConcertSync/ConcertSyncServer/Source/ConcertSyncServer/Private/ConcertSyncServer.cpp b/Engine/Plugins/Developer/Concert/ConcertSync/ConcertSyncServer/Source/ConcertSyncServer/Private/ConcertSyncServer.cpp
--- a/Engine/Plugins/Developer/Concert/ConcertSync/ConcertSyncServer/Source/ConcertSyncServer/Private/ConcertSyncServer.cpp
+++ b/Engine/Plugins/Developer/Concert/ConcertSync/ConcertSyncServer/Source/ConcertSyncServer/Private/ConcertSyncServer.cpp
@@ -261,6 +261,24 @@ void FConcertSyncServer::SetFileSharingService(TSharedPtr<IConcertFileSharingSer
 	FileSharingService = MoveTemp(InFileSharingService);
 }
 
+TOptional<FConcertSyncSessionDatabaseNonNullPtr> FConcertSyncServer::GetLiveSessionDatabase(const FGuid& SessionId)
+{
+	if (const TSharedPtr<FConcertSyncServerLiveSession>* Session = LiveSessions.Find(SessionId))
+	{
+		return { &Session->Get()->GetSessionDatabase() };
+	}
+	return {};
+}
+
+TOptional<FConcertSyncSessionDatabaseNonNullPtr> FConcertSyncServer::GetArchivedSessionDatabase(const FGuid& SessionId)
+{
+	if (const TSharedPtr<FConcertSyncServerArchivedSession>* Session = ArchivedSessions.Find(SessionId))
+	{
+		return { &Session->Get()->GetSessionDatabase() };
+	}
+	return {};
+}
+
 void FConcertSyncServer::GetSessionsFromPathImpl(const IConcertServer& InServer, const FString& InPath, TArray<FConcertSessionInfo>& OutSessionInfos, TArray<FDateTime>* OutSessionCreationTimes)
 {
 	IFileManager::Get().IterateDirectory(*InPath, [&OutSessionInfos, &OutSessionCreationTimes](const TCHAR* FilenameOrDirectory, bool bIsDirectory) -> bool
diff --git a/Engine/Plugins/Developer/Concert/ConcertSync/ConcertSyncServer/Source/ConcertSyncServer/Private/ConcertSyncServer.h b/Engine/Plugins/Developer/Concert/ConcertSync/ConcertSyncServer/Source/ConcertSyncServer/Private/ConcertSyncServer.h
--- a/Engine/Plugins/Developer/Concert/ConcertSync/ConcertSyncServer/Source/ConcertSyncServer/Private/ConcertSyncServer.h
+++ b/Engine/Plugins/Developer/Concert/ConcertSync/ConcertSyncServer/Source/ConcertSyncServer/Private/ConcertSyncServer.h
@@ -29,6 +29,8 @@ public:
 	virtual void Shutdown() override;
 	virtual IConcertServerRef GetConcertServer() const override;
 	virtual void SetFileSharingService(TSharedPtr<IConcertFileSharingService> InFileSharingService) override;
+	virtual TOptional<FConcertSyncSessionDatabaseNonNullPtr> GetLiveSessionDatabase(const FGuid& SessionId) override;
+	virtual TOptional<FConcertSyncSessionDatabaseNonNullPtr> GetArchivedSessionDatabase(const FGuid& SessionId) override;
 
 	//~ IConcertServerEventSink interface
 	void GetSessionsFromPathImpl(const IConcertServer& InServer, const FString& InPath, TArray<FConcertSessionInfo>& OutSessionInfos, TArray<FDateTime>* OutSessionCreationTimes = nullptr);
diff --git a/Engine/Plugins/Developer/Concert/ConcertSync/ConcertSyncServer/Source/ConcertSyncServer/Public/IConcertSyncServer.h b/Engine/Plugins/Developer/Concert/ConcertSync/ConcertSyncServer/Source/ConcertSyncServer/Public/IConcertSyncServer.h
--- a/Engine/Plugins/Developer/Concert/ConcertSync/ConcertSyncServer/Source/ConcertSyncServer/Public/IConcertSyncServer.h
+++ b/Engine/Plugins/Developer/Concert/ConcertSync/ConcertSyncServer/Source/ConcertSyncServer/Public/IConcertSyncServer.h
@@ -5,10 +5,14 @@
 #include "CoreMinimal.h"
 #include "IConcertModule.h"
 #include "ConcertSyncSessionFlags.h"
-class IConcertFileSharingService;
+#include "Templates/NonNullPointer.h"
 
+class FConcertSyncSessionDatabase;
+class IConcertFileSharingService;
 class UConcertServerConfig;
 
+using FConcertSyncSessionDatabaseNonNullPtr = TNonNullPtr<FConcertSyncSessionDatabase>;
+
 /**
  * Interface for a Concert Sync Server.
  */
@@ -28,4 +32,10 @@ public:
 
 	/** Set the file sharing service, enabling the server to work with large files. The server sharing service must be compatible with the client one. */
 	virtual void SetFileSharingService(TSharedPtr<IConcertFileSharingService> InFileSharingService) = 0;
+
+	/** Gets the session data base for the given live session if it exists */
+	virtual TOptional<FConcertSyncSessionDatabaseNonNullPtr> GetLiveSessionDatabase(const FGuid& SessionId) = 0;
+
+	/** Gets the session data base for the given archived session if it exists */
+	virtual TOptional<FConcertSyncSessionDatabaseNonNullPtr> GetArchivedSessionDatabase(const FGuid& SessionId) = 0;
 };
diff --git a/Engine/Plugins/Developer/Concert/ConcertUI/ConcertServerUI/Source/ConcertServerUI/Private/Widgets/ConcertServerWindowController.cpp b/Engine/Plugins/Developer/Concert/ConcertUI/ConcertServerUI/Source/ConcertServerUI/Private/Widgets/ConcertServerWindowController.cpp
--- a/Engine/Plugins/Developer/Concert/ConcertUI/ConcertServerUI/Source/ConcertServerUI/Private/Widgets/ConcertServerWindowController.cpp
+++ b/Engine/Plugins/Developer/Concert/ConcertUI/ConcertServerUI/Source/ConcertServerUI/Private/Widgets/ConcertServerWindowController.cpp
@@ -95,7 +95,7 @@ TSharedPtr<FConcertServerSessionTab> FConcertServerWindowController::GetOrRegist
 			return *FoundId;
 		}
 	
-		const TSharedRef<FConcertServerSessionTab> SessionTab = MakeShared<FConcertServerSessionTab>(Session.ToSharedRef(), RootWindow.ToSharedRef());
+		const TSharedRef<FConcertServerSessionTab> SessionTab = MakeShared<FConcertServerSessionTab>(Session.ToSharedRef(), ServerInstance.ToSharedRef(), RootWindow.ToSharedRef());
 		RegisteredSessions.Add(SessionInfo.SessionId, SessionTab);
 		return SessionTab;
 	}
diff --git a/Engine/Plugins/Developer/Concert/ConcertUI/ConcertServerUI/Source/ConcertServerUI/Private/Widgets/Session/ConcertServerSessionTab.cpp b/Engine/Plugins/Developer/Concert/ConcertUI/ConcertServerUI/Source/ConcertServerUI/Private/Widgets/Session/ConcertServerSessionTab.cpp
--- a/Engine/Plugins/Developer/Concert/ConcertUI/ConcertServerUI/Source/ConcertServerUI/Private/Widgets/Session/ConcertServerSessionTab.cpp
+++ b/Engine/Plugins/Developer/Concert/ConcertUI/ConcertServerUI/Source/ConcertServerUI/Private/Widgets/Session/ConcertServerSessionTab.cpp
@@ -9,9 +9,10 @@
 #include "Widgets/StatusBar/SConcertStatusBar.h"
 #include "Widgets/Session/SConcertSessionInspector.h"
 
-FConcertServerSessionTab::FConcertServerSessionTab(TSharedRef<IConcertServerSession> InspectedSession, const TSharedRef<SWindow>& ConstructUnderWindow)
+FConcertServerSessionTab::FConcertServerSessionTab(TSharedRef<IConcertServerSession> InspectedSession, TSharedRef<IConcertSyncServer> SyncServer, const TSharedRef<SWindow>& ConstructUnderWindow)
 	: InspectedSession(MoveTemp(InspectedSession))
-	, DockTab(CreateTab(InspectedSession, ConstructUnderWindow))
+	, SessionHistoryController(MakeShared<FServerSessionHistoryController>(InspectedSession, SyncServer))
+	, DockTab(CreateTab(ConstructUnderWindow))
 {}
 
 void FConcertServerSessionTab::OpenSessionTab() const
@@ -29,20 +30,20 @@ void FConcertServerSessionTab::OpenSessionTab() const
 	}
 }
 
-TSharedRef<SDockTab> FConcertServerSessionTab::CreateTab(const TSharedRef<IConcertServerSession>& InspectedSession, const TSharedRef<SWindow>& ConstructUnderWindow)
+TSharedRef<SDockTab> FConcertServerSessionTab::CreateTab(const TSharedRef<SWindow>& ConstructUnderWindow) const
 {
 	const FText Title = FText::FromString(InspectedSession->GetSessionInfo().SessionName);
-	TSharedRef<SDockTab> DockTab = SNew(SDockTab)
+	TSharedRef<SDockTab> NewDockTab = SNew(SDockTab)
 		.Label(Title)
 		.TabRole(MajorTab);
-	DockTab->SetContent(
-		SNew(SConcertSessionInspector, DockTab, ConstructUnderWindow)
+	NewDockTab->SetContent(
+		SNew(SConcertSessionInspector, SConcertSessionInspector::FRequiredArgs{NewDockTab, ConstructUnderWindow, SessionHistoryController})
 			.StatusBar()
 			[
 				SNew(SConcertStatusBar, *GetTabPlayerHolderId(InspectedSession))
 			]
 		);
-	return DockTab;
+	return NewDockTab;
 }
 
 FString FConcertServerSessionTab::GetTabPlayerHolderId(const TSharedRef<IConcertServerSession>& InspectedSession)
diff --git a/Engine/Plugins/Developer/Concert/ConcertUI/ConcertServerUI/Source/ConcertServerUI/Private/Widgets/Session/ConcertServerSessionTab.h b/Engine/Plugins/Developer/Concert/ConcertUI/ConcertServerUI/Source/ConcertServerUI/Private/Widgets/Session/ConcertServerSessionTab.h
--- a/Engine/Plugins/Developer/Concert/ConcertUI/ConcertServerUI/Source/ConcertServerUI/Private/Widgets/Session/ConcertServerSessionTab.h
+++ b/Engine/Plugins/Developer/Concert/ConcertUI/ConcertServerUI/Source/ConcertServerUI/Private/Widgets/Session/ConcertServerSessionTab.h
@@ -3,17 +3,23 @@
 #pragma once
 
 #include "CoreMinimal.h"
+#include "ServerSessionHistoryController.h"
 
+class IConcertSyncServer;
+class FServerSessionHistoryController;
 class IConcertServerSession;
 class SDockTab;
 class SWindow;
 
-/** Manages the tab that contains the UI for a session */
+/**
+ * Manages the tab that contains the UI for a session.
+ * It has access to the controllers and views needed for displaying a session.
+ */
 class FConcertServerSessionTab
 {
 public:
 
-	FConcertServerSessionTab(TSharedRef<IConcertServerSession> InspectedSession, const TSharedRef<SWindow>& ConstructUnderWindow);
+	FConcertServerSessionTab(TSharedRef<IConcertServerSession> InspectedSession, TSharedRef<IConcertSyncServer> SyncServer, const TSharedRef<SWindow>& ConstructUnderWindow);
 
 	/** Opens or draws attention to the given tab */
 	void OpenSessionTab() const;
@@ -23,11 +29,14 @@ private:
 	/** The session being inspected */
 	const TSharedRef<IConcertServerSession> InspectedSession;
 	
+	const TSharedRef<FServerSessionHistoryController> SessionHistoryController;
+	
 	/** The tab containing the UI for InspectedSession */
 	const TSharedRef<SDockTab> DockTab;
 
+
 	/** Creates a tab widget */
-	static TSharedRef<SDockTab> CreateTab(const TSharedRef<IConcertServerSession>& InspectedSession, const TSharedRef<SWindow>& ConstructUnderWindow);
+	TSharedRef<SDockTab> CreateTab(const TSharedRef<SWindow>& ConstructUnderWindow) const;
 	
 	/** Generates a tab ID for FTabManager::InsertNewDocumentTab */
 	static FString GetTabPlayerHolderId(const TSharedRef<IConcertServerSession>& InspectedSession);
diff --git a/Engine/Plugins/Developer/Concert/ConcertUI/ConcertServerUI/Source/ConcertServerUI/Private/Widgets/Session/SConcertSessionInspector.cpp b/Engine/Plugins/Developer/Concert/ConcertUI/ConcertServerUI/Source/ConcertServerUI/Private/Widgets/Session/SConcertSessionInspector.cpp
--- a/Engine/Plugins/Developer/Concert/ConcertUI/ConcertServerUI/Source/ConcertServerUI/Private/Widgets/Session/SConcertSessionInspector.cpp
+++ b/Engine/Plugins/Developer/Concert/ConcertUI/ConcertServerUI/Source/ConcertServerUI/Private/Widgets/Session/SConcertSessionInspector.cpp
@@ -3,6 +3,7 @@
 #include "Widgets/Session/SConcertSessionInspector.h"
 
 #include "Framework/Docking/TabManager.h"
+#include "SessionHistory/SSessionHistory.h"
 #include "Widgets/Docking/SDockTab.h"
 #include "Widgets/SWindow.h"
 
@@ -12,7 +13,7 @@ const FName SConcertSessionInspector::HistoryTabId("HistoryTabId");
 const FName SConcertSessionInspector::SessionContentTabId("SessionContentTabId");
 const FName SConcertSessionInspector::ConnectionMonitorTabId("ConnectionMonitorTabId");
 
-void SConcertSessionInspector::Construct(const FArguments& InArgs, const TSharedRef<SDockTab>& ConstructUnderMajorTab, const TSharedRef<SWindow>& ConstructUnderWindow)
+void SConcertSessionInspector::Construct(const FArguments& InArgs, const FRequiredArgs& RequiredArgs)
 {
 	ChildSlot
 	[
@@ -37,7 +38,7 @@ void SConcertSessionInspector::Construct(const FArguments& InArgs, const TShared
 				.BorderImage(FAppStyle::Get().GetBrush("Brushes.Title"))
 				.Padding(FMargin(0.f, 0.f, 0.f, 5.0f)) // Visually separate status bar from tabs
 				[
-					CreateTabs(ConstructUnderMajorTab, ConstructUnderWindow)
+					CreateTabs(InArgs, RequiredArgs)
 				]
 			]
 
@@ -52,11 +53,11 @@ void SConcertSessionInspector::Construct(const FArguments& InArgs, const TShared
 	];
 }
 
-TSharedRef<SWidget> SConcertSessionInspector::CreateTabs(const TSharedRef<SDockTab>& ConstructUnderMajorTab, const TSharedRef<SWindow>& ConstructUnderWindow)
+TSharedRef<SWidget> SConcertSessionInspector::CreateTabs(const FArguments& InArgs, const FRequiredArgs& RequiredArgs)
 {
-	TabManager = FGlobalTabmanager::Get()->NewTabManager(ConstructUnderMajorTab);
+	TabManager = FGlobalTabmanager::Get()->NewTabManager(RequiredArgs.ConstructUnderMajorTab);
 	
-	TabManager->RegisterTabSpawner(HistoryTabId, FOnSpawnTab::CreateSP(this, &SConcertSessionInspector::SpawnActivityHistory))
+	TabManager->RegisterTabSpawner(HistoryTabId, FOnSpawnTab::CreateSP(this, &SConcertSessionInspector::SpawnActivityHistory, RequiredArgs.SessionHistoryController))
 		.SetDisplayName(LOCTEXT("ActivityHistoryLabel", "History"));
 	
 	TabManager->RegisterTabSpawner(SessionContentTabId, FOnSpawnTab::CreateSP(this, &SConcertSessionInspector::SpawnSessionContent))
@@ -96,16 +97,16 @@ TSharedRef<SWidget> SConcertSessionInspector::CreateTabs(const TSharedRef<SDockT
 				)
 		);
 
-	return TabManager->RestoreFrom(Layout, ConstructUnderWindow).ToSharedRef();
+	return TabManager->RestoreFrom(Layout, RequiredArgs.ConstructUnderWindow).ToSharedRef();
 }
 
-TSharedRef<SDockTab> SConcertSessionInspector::SpawnActivityHistory(const FSpawnTabArgs& Args)
+TSharedRef<SDockTab> SConcertSessionInspector::SpawnActivityHistory(const FSpawnTabArgs& Args, TSharedRef<FServerSessionHistoryController> SessionHistoryController)
 {
 	return SNew(SDockTab)
 		.Label(LOCTEXT("ActivityHistoryLabel", "History"))
 		.TabRole(PanelTab)
 		[
-			SNullWidget::NullWidget
+			SessionHistoryController->GetSessionHistory()
 		];
 }
 
diff --git a/Engine/Plugins/Developer/Concert/ConcertUI/ConcertServerUI/Source/ConcertServerUI/Private/Widgets/Session/SConcertSessionInspector.h b/Engine/Plugins/Developer/Concert/ConcertUI/ConcertServerUI/Source/ConcertServerUI/Private/Widgets/Session/SConcertSessionInspector.h
--- a/Engine/Plugins/Developer/Concert/ConcertUI/ConcertServerUI/Source/ConcertServerUI/Private/Widgets/Session/SConcertSessionInspector.h
+++ b/Engine/Plugins/Developer/Concert/ConcertUI/ConcertServerUI/Source/ConcertServerUI/Private/Widgets/Session/SConcertSessionInspector.h
@@ -3,6 +3,7 @@
 #pragma once
 
 #include "CoreMinimal.h"
+#include "ServerSessionHistoryController.h"
 #include "Widgets/SCompoundWidget.h"
 #include "Widgets/DeclarativeSyntaxSupport.h"
 
@@ -22,6 +23,19 @@ class SConcertSessionInspector : public SCompoundWidget
 {
 public:
 
+	struct FRequiredArgs
+	{
+		TSharedRef<SDockTab> ConstructUnderMajorTab;
+		TSharedRef<SWindow> ConstructUnderWindow;
+		TSharedRef<FServerSessionHistoryController> SessionHistoryController;
+
+		FRequiredArgs(TSharedRef<SDockTab> ConstructUnderMajorTab, TSharedRef<SWindow> ConstructUnderWindow,TSharedRef<FServerSessionHistoryController> SessionHistoryController)
+			: ConstructUnderMajorTab(MoveTemp(ConstructUnderMajorTab))
+			, ConstructUnderWindow(MoveTemp(ConstructUnderWindow))
+			, SessionHistoryController(MoveTemp(SessionHistoryController))
+		{}
+	};
+
 	static const FName HistoryTabId;
 	static const FName SessionContentTabId;
 	static const FName ConnectionMonitorTabId;
@@ -30,7 +44,7 @@ public:
 		SLATE_NAMED_SLOT(FArguments, StatusBar)
 	SLATE_END_ARGS()
 
-	void Construct(const FArguments& InArgs, const TSharedRef<SDockTab>& ConstructUnderMajorTab, const TSharedRef<SWindow>& ConstructUnderWindow);
+	void Construct(const FArguments& InArgs, const FRequiredArgs& RequiredArgs);
 
 private:
 	
@@ -38,8 +52,8 @@ private:
 	TSharedPtr<FTabManager> TabManager;
 
 	// Spawn tabs
-	TSharedRef<SWidget> CreateTabs(const TSharedRef<SDockTab>& ConstructUnderMajorTab, const TSharedRef<SWindow>& ConstructUnderWindow);
-	TSharedRef<SDockTab> SpawnActivityHistory(const FSpawnTabArgs& Args);
+	TSharedRef<SWidget> CreateTabs(const FArguments& InArgs, const FRequiredArgs& RequiredArgs);
+	TSharedRef<SDockTab> SpawnActivityHistory(const FSpawnTabArgs& Args, TSharedRef<FServerSessionHistoryController> SessionHistoryController);
 	TSharedRef<SDockTab> SpawnSessionContent(const FSpawnTabArgs& Args);
 	TSharedRef<SDockTab> SpawnConnectionMonitor(const FSpawnTabArgs& Args);
 };
