Unshelved from pending changelist '19219261':

Add dialog box when archiving or deleting session

diff --git a/Engine/Plugins/Developer/Concert/ConcertUI/ConcertServerUI/Source/ConcertServerUI/Private/Widgets/Browser/ConcertServerSessionBrowserController.cpp b/Engine/Plugins/Developer/Concert/ConcertUI/ConcertServerUI/Source/ConcertServerUI/Private/Widgets/Browser/ConcertServerSessionBrowserController.cpp
--- a/Engine/Plugins/Developer/Concert/ConcertUI/ConcertServerUI/Source/ConcertServerUI/Private/Widgets/Browser/ConcertServerSessionBrowserController.cpp
+++ b/Engine/Plugins/Developer/Concert/ConcertUI/ConcertServerUI/Source/ConcertServerUI/Private/Widgets/Browser/ConcertServerSessionBrowserController.cpp
@@ -18,6 +18,12 @@
 
 #define LOCTEXT_NAMESPACE "UnrealMultiUserUI"
 
+int32 FConcertServerSessionBrowserController::GetNumConnectedClients(const FGuid& SessionId) const
+{
+	TSharedPtr<IConcertServerSession> Session = ServerInstance->GetConcertServer()->GetLiveSession(SessionId);
+	return ensure(Session) ? Session->GetSessionClients().Num() : 0;
+}
+
 void FConcertServerSessionBrowserController::Init(const FConcertComponentInitParams& Params)
 {
 	ServerInstance = Params.Server;
@@ -120,6 +126,25 @@ void FConcertServerSessionBrowserController::RestoreSession(const FGuid& ServerA
 	}
 }
 
+TSharedRef<SDockTab> FConcertServerSessionBrowserController::SpawnSessionBrowserTab(const FSpawnTabArgs& Args)
+{
+	const TSharedRef<SDockTab> DockTab = SNew(SDockTab)
+		.Label(LOCTEXT("SessionBrowserTabTitle", "Sessions"))
+		.TabRole(MajorTab)
+		[
+			SAssignNew(ConcertBrowser, SConcertServerSessionBrowser, SharedThis(this))
+				.DoubleClickSession(this, &FConcertServerSessionBrowserController::OpenSession)
+		];
+
+	FGlobalTabmanager::Get()->SetMainTab(DockTab);
+	return DockTab;
+}
+
+void FConcertServerSessionBrowserController::OpenSession(TSharedPtr<FConcertSessionItem> SessionItem)
+{
+	Owner.Pin()->OpenSessionTab(SessionItem->SessionId);
+}
+
 void FConcertServerSessionBrowserController::RenameSession(const FGuid& ServerAdminEndpointId, const FGuid& SessionId, const FString& NewName)
 {
 	if (const TOptional<FConcertSessionInfo> SessionInfo = GetArchivedSessionInfo(ServerAdminEndpointId, SessionId))
@@ -136,36 +161,20 @@ void FConcertServerSessionBrowserController::RenameSession(const FGuid& ServerAd
 
 void FConcertServerSessionBrowserController::DeleteSession(const FGuid& ServerAdminEndpointId, const FGuid& SessionId)
 {
-	if (const TOptional<FConcertSessionInfo> SessionInfo = GetActiveSessionInfo(ServerAdminEndpointId, SessionId))
+	const TOptional<FConcertSessionInfo> ArchivedInfo = GetArchivedSessionInfo(ServerAdminEndpointId, SessionId);
+	const TOptional<FConcertSessionInfo> LiveInfo = GetActiveSessionInfo(ServerAdminEndpointId, SessionId);
+	const TOptional<FString> SessionName = ArchivedInfo ? ArchivedInfo->SessionName : LiveInfo ? LiveInfo->SessionName : TOptional<FString>{};
+	if (SessionName)
 	{
 		FText FailureReason;
 		const bool bSuccess = ServerInstance->GetConcertServer()->DestroySession(SessionId, FailureReason);
 		NotifyUserOfFinishedSessionAction(bSuccess,
-			bSuccess ? FText::Format(LOCTEXT("ArchivedSessionFmt", "Deleted Session '{0}'"), FText::FromString(SessionInfo->SessionName)) : FText::Format(LOCTEXT("FailedToArchivedSessionFmt", "Failed to archive Session '{0}'"), FText::FromString(SessionInfo->SessionName)),
+			bSuccess ? FText::Format(LOCTEXT("ArchivedSessionFmt", "Deleted Session '{0}'"), FText::FromString(*SessionName)) : FText::Format(LOCTEXT("FailedToArchivedSessionFmt", "Failed to archive Session '{0}'"), FText::FromString(LiveInfo->SessionName)),
 			FailureReason
 			);
 	}
 }
 
-TSharedRef<SDockTab> FConcertServerSessionBrowserController::SpawnSessionBrowserTab(const FSpawnTabArgs& Args)
-{
-	const TSharedRef<SDockTab> DockTab = SNew(SDockTab)
-		.Label(LOCTEXT("SessionBrowserTabTitle", "Sessions"))
-		.TabRole(MajorTab)
-		[
-			SAssignNew(ConcertBrowser, SConcertServerSessionBrowser, SharedThis(this))
-				.DoubleClickSession(this, &FConcertServerSessionBrowserController::OpenSession)
-		];
-
-	FGlobalTabmanager::Get()->SetMainTab(DockTab);
-	return DockTab;
-}
-
-void FConcertServerSessionBrowserController::OpenSession(TSharedPtr<FConcertSessionItem> SessionItem)
-{
-	Owner.Pin()->OpenSessionTab(SessionItem->SessionId);
-}
-
 void FConcertServerSessionBrowserController::NotifyUserOfFinishedSessionAction(const bool bSuccess, const FText& Title, const FText& Details)
 {
 	FNotificationInfo NotificationInfo(Title);
diff --git a/Engine/Plugins/Developer/Concert/ConcertUI/ConcertServerUI/Source/ConcertServerUI/Private/Widgets/Browser/ConcertServerSessionBrowserController.h b/Engine/Plugins/Developer/Concert/ConcertUI/ConcertServerUI/Source/ConcertServerUI/Private/Widgets/Browser/ConcertServerSessionBrowserController.h
--- a/Engine/Plugins/Developer/Concert/ConcertUI/ConcertServerUI/Source/ConcertServerUI/Private/Widgets/Browser/ConcertServerSessionBrowserController.h
+++ b/Engine/Plugins/Developer/Concert/ConcertUI/ConcertServerUI/Source/ConcertServerUI/Private/Widgets/Browser/ConcertServerSessionBrowserController.h
@@ -20,6 +20,8 @@ class FConcertServerSessionBrowserController
 {
 public:
 
+	int32 GetNumConnectedClients(const FGuid& SessionId) const;
+	
 	//~ Begin IConcertComponent Interface
 	virtual void Init(const FConcertComponentInitParams& Params) override;
 	//~ End IConcertComponent Interface
@@ -55,12 +57,14 @@ private:
 	
 	TSharedPtr<SConcertServerSessionBrowser> ConcertBrowser;
 	
-	void RenameSession(const FGuid& ServerAdminEndpointId, const FGuid& SessionId, const FString& NewName);
-	void DeleteSession(const FGuid& ServerAdminEndpointId, const FGuid& SessionId);
-	
 	TSharedRef<SDockTab> SpawnSessionBrowserTab(const FSpawnTabArgs& Args);
+
+	// Session actions
 	void OpenSession(TSharedPtr<FConcertSessionItem> SessionItem);
-	
+	void RenameSession(const FGuid& ServerAdminEndpointId, const FGuid& SessionId, const FString& NewName);
+	void DeleteSession(const FGuid& ServerAdminEndpointId, const FGuid& SessionId);
+
+	// Notifications
 	void NotifyUserOfFinishedSessionAction(const bool bSuccess, const FText& Title) { NotifyUserOfFinishedSessionAction(bSuccess, Title, FText::GetEmpty()); }
     /** Shared logic for displaying a notification after creating, archiving, restoring, renaming and deleting a session. */
 	void NotifyUserOfFinishedSessionAction(const bool bSuccess, const FText& Title, const FText& Details);
diff --git a/Engine/Plugins/Developer/Concert/ConcertUI/ConcertServerUI/Source/ConcertServerUI/Private/Widgets/Browser/SConcertServerSessionBrowser.cpp b/Engine/Plugins/Developer/Concert/ConcertUI/ConcertServerUI/Source/ConcertServerUI/Private/Widgets/Browser/SConcertServerSessionBrowser.cpp
--- a/Engine/Plugins/Developer/Concert/ConcertUI/ConcertServerUI/Source/ConcertServerUI/Private/Widgets/Browser/SConcertServerSessionBrowser.cpp
+++ b/Engine/Plugins/Developer/Concert/ConcertUI/ConcertServerUI/Source/ConcertServerUI/Private/Widgets/Browser/SConcertServerSessionBrowser.cpp
@@ -7,6 +7,7 @@
 #include "Widgets/StatusBar/SConcertStatusBar.h"
 
 #include "SPositiveActionButton.h"
+#include "Dialog/SMessageDialog.h"
 #include "SessionBrowser/ConcertSessionItem.h"
 #include "Widgets/ConcertServerTabs.h"
 #include "Widgets/Input/SSearchBox.h"
@@ -15,7 +16,7 @@
 
 #define LOCTEXT_NAMESPACE "UnrealMultiUserUI"
 
-void SConcertServerSessionBrowser::Construct(const FArguments& InArgs, TSharedRef<IConcertSessionBrowserController> InController)
+void SConcertServerSessionBrowser::Construct(const FArguments& InArgs, TSharedRef<FConcertServerSessionBrowserController> InController)
 {
 	Controller = InController;
 	ChildSlot
@@ -49,7 +50,45 @@ TSharedRef<SWidget> SConcertServerSessionBrowser::MakeSessionTableView(const FAr
 {
 	SearchText = MakeShared<FText>();
 	return SAssignNew(SessionBrowser, SConcertSessionBrowser, Controller.Pin().ToSharedRef(), SearchText)
-		.OnSessionDoubleClicked(InArgs._DoubleClickSession);
+		.OnSessionDoubleClicked(InArgs._DoubleClickSession)
+		.CanArchiveSession(this, &SConcertServerSessionBrowser::ConfirmArchiveOperationWithDialog)
+		.CanDeleteSession(this, &SConcertServerSessionBrowser::ConfirmDeleteOperationWithDialog);
+}
+
+bool SConcertServerSessionBrowser::ConfirmArchiveOperationWithDialog(TSharedPtr<FConcertSessionItem> SessionItem)
+{
+	const int32 NumUsers = Controller.Pin()->GetNumConnectedClients(SessionItem->SessionId);
+	const FText Message = FText::Format(
+		LOCTEXT("ArchiveDescription", "There {0}|plural(one=is,other=are) {0} connected {0}|plural(one=client,other=clients) in the current session.\nArchiving a session will force all connected clients to disconnect."),
+		NumUsers
+		);
+
+	constexpr int32 ArchiveIndex = 0;
+	const TSharedRef<SMessageDialog> Dialog = SNew(SMessageDialog)
+		.Title(LOCTEXT("DisconnectUsersTitle", "Force Users to Disconnect?"))
+		.IconBrush("Icons.WarningWithColor.Large")
+		.Message(Message)
+		.Buttons({
+			SMessageDialog::FButton(LOCTEXT("ArchiveButton", "Archive")).SetPrimary(true),
+			SMessageDialog::FButton(LOCTEXT("CancelButton", "Cancel"))
+		});
+	return Dialog->ShowModal() == ArchiveIndex;
+}
+
+bool SConcertServerSessionBrowser::ConfirmDeleteOperationWithDialog(TSharedPtr<FConcertSessionItem> SessionItem)
+{
+	const FText Message = LOCTEXT("DeleteDescription", "Deleting a session will cause all associated data to be removed.");
+	
+	constexpr int32 DeleteIndex = 0;
+	const TSharedRef<SMessageDialog> Dialog = SNew(SMessageDialog)
+		.Title(LOCTEXT("DisconnectUsersTitle", "Delete session?"))
+		.IconBrush("Icons.WarningWithColor.Large")
+		.Message(Message)
+		.Buttons({
+			SMessageDialog::FButton(LOCTEXT("DeleteButton", "Delete")),
+			SMessageDialog::FButton(LOCTEXT("CancelButton", "Cancel")).SetPrimary(true)
+		});
+	return Dialog->ShowModal() == DeleteIndex;
 }
 
 #undef LOCTEXT_NAMESPACE
diff --git a/Engine/Plugins/Developer/Concert/ConcertUI/ConcertServerUI/Source/ConcertServerUI/Private/Widgets/Browser/SConcertServerSessionBrowser.h b/Engine/Plugins/Developer/Concert/ConcertUI/ConcertServerUI/Source/ConcertServerUI/Private/Widgets/Browser/SConcertServerSessionBrowser.h
--- a/Engine/Plugins/Developer/Concert/ConcertUI/ConcertServerUI/Source/ConcertServerUI/Private/Widgets/Browser/SConcertServerSessionBrowser.h
+++ b/Engine/Plugins/Developer/Concert/ConcertUI/ConcertServerUI/Source/ConcertServerUI/Private/Widgets/Browser/SConcertServerSessionBrowser.h
@@ -19,17 +19,20 @@ public:
 		SLATE_EVENT(FSessionDelegate, DoubleClickSession)
 	SLATE_END_ARGS()
 
-	void Construct(const FArguments& InArgs, TSharedRef<IConcertSessionBrowserController> InController);
+	void Construct(const FArguments& InArgs, TSharedRef<FConcertServerSessionBrowserController> InController);
 
 	void RefreshSessionList() { SessionBrowser->RefreshSessionList(); }
 	
 private:
 
 	/** We can ask the controller about information and notify it about UI events. */
-	TWeakPtr<IConcertSessionBrowserController> Controller;
+	TWeakPtr<FConcertServerSessionBrowserController> Controller;
 
 	TSharedPtr<FText> SearchText;
 	TSharedPtr<SConcertSessionBrowser> SessionBrowser;
 	
 	TSharedRef<SWidget> MakeSessionTableView(const FArguments& InArgs);
+
+	bool ConfirmArchiveOperationWithDialog(TSharedPtr<FConcertSessionItem> SessionItem);
+	bool ConfirmDeleteOperationWithDialog(TSharedPtr<FConcertSessionItem> SessionItem);
 };
diff --git a/Engine/Plugins/Developer/Concert/ConcertUI/ConcertSharedSlate/Source/ConcertSharedSlate/Private/SessionBrowser/SConcertSessionBrowser.cpp b/Engine/Plugins/Developer/Concert/ConcertUI/ConcertSharedSlate/Source/ConcertSharedSlate/Private/SessionBrowser/SConcertSessionBrowser.cpp
--- a/Engine/Plugins/Developer/Concert/ConcertUI/ConcertSharedSlate/Source/ConcertSharedSlate/Private/SessionBrowser/SConcertSessionBrowser.cpp
+++ b/Engine/Plugins/Developer/Concert/ConcertUI/ConcertSharedSlate/Source/ConcertSharedSlate/Private/SessionBrowser/SConcertSessionBrowser.cpp
@@ -51,6 +51,8 @@ void SConcertSessionBrowser::Construct(const FArguments& InArgs, TSharedRef<ICon
 	OnSessionClicked = InArgs._OnSessionClicked;
 	OnSessionDoubleClicked = InArgs._OnSessionDoubleClicked;
 	OnRequestedDeleteSession = InArgs._OnRequestedDeleteSession;
+	CanArchiveSession = InArgs._CanArchiveSession;
+	CanDeleteSession = InArgs._CanDeleteSession;
 
 	// Setup search filter.
 	SearchedText = InSearchText; // Reload a previous search text (in any). Useful to remember searched text between join/leave sessions, but not persistent if the tab is closed.
@@ -1026,19 +1028,22 @@ void SConcertSessionBrowser::RequestDeleteSession(const TSharedPtr<FConcertSessi
 	const FText SeverNameInText = FText::FromString(DeletedItem->ServerName);
 	const FText ConfirmationMessage = FText::Format(LOCTEXT("DeleteSessionConfirmationMessage", "Do you really want to delete the session \"{0}\" from the server \"{1}\"?"), SessionNameInText, SeverNameInText);
 	const FText ConfirmationTitle = LOCTEXT("DeleteSessionConfirmationTitle", "Delete Session Confirmation");
-
-	if (FMessageDialog::Open(EAppMsgType::YesNo, ConfirmationMessage, &ConfirmationTitle) == EAppReturnType::Yes) // Confirmed?
+	
+	if (DeletedItem->Type == FConcertSessionItem::EType::ActiveSession)
 	{
-		if (DeletedItem->Type == FConcertSessionItem::EType::ActiveSession)
+		if (!CanArchiveSession.IsBound() || CanArchiveSession.Execute(DeletedItem))
 		{
 			GetController()->DeleteActiveSession(DeletedItem->ServerAdminEndpointId, DeletedItem->SessionId);
+			OnRequestedDeleteSession.ExecuteIfBound(DeletedItem);
 		}
-		else if (DeletedItem->Type == FConcertSessionItem::EType::ArchivedSession)
+	}
+	else if (DeletedItem->Type == FConcertSessionItem::EType::ArchivedSession)
+	{
+		if (!CanDeleteSession.IsBound() || CanDeleteSession.Execute(DeletedItem))
 		{
 			GetController()->DeleteArchivedSession(DeletedItem->ServerAdminEndpointId, DeletedItem->SessionId);
+			OnRequestedDeleteSession.ExecuteIfBound(DeletedItem);
 		}
-
-		OnRequestedDeleteSession.ExecuteIfBound(DeletedItem);
 	}
 }
 
diff --git a/Engine/Plugins/Developer/Concert/ConcertUI/ConcertSharedSlate/Source/ConcertSharedSlate/Public/SessionBrowser/SConcertSessionBrowser.h b/Engine/Plugins/Developer/Concert/ConcertUI/ConcertSharedSlate/Source/ConcertSharedSlate/Public/SessionBrowser/SConcertSessionBrowser.h
--- a/Engine/Plugins/Developer/Concert/ConcertUI/ConcertSharedSlate/Source/ConcertSharedSlate/Public/SessionBrowser/SConcertSessionBrowser.h
+++ b/Engine/Plugins/Developer/Concert/ConcertUI/ConcertSharedSlate/Source/ConcertSharedSlate/Public/SessionBrowser/SConcertSessionBrowser.h
@@ -25,6 +25,7 @@ DECLARE_DELEGATE_RetVal_OneParam(TSharedRef<SWidget>, FExtendSessionTable, TShar
 DECLARE_DELEGATE_OneParam(FExtenderDelegate, FExtender&);
 DECLARE_DELEGATE_TwoParams(FExtendSessionContextMenu, TSharedPtr<FConcertSessionItem>, FExtender&)
 DECLARE_DELEGATE_OneParam(FSessionDelegate, TSharedPtr<FConcertSessionItem> /*Session*/);
+DECLARE_DELEGATE_RetVal_OneParam(bool, FCanRemoveSession, TSharedPtr<FConcertSessionItem> /* SessionItem */);
 
 /**
  * Enables the user to browse/search/filter/sort active and archived sessions, create new session,
@@ -72,6 +73,11 @@ public:
 	SLATE_EVENT(FSessionDelegate, OnSessionDoubleClicked)
 	/** Called after a user has requested to delete a session */
 	SLATE_EVENT(FSessionDelegate, OnRequestedDeleteSession)
+
+	/** Ask the user to confirm archiving - most obvious implementation is showing a dialog box */
+	SLATE_EVENT(FCanRemoveSession, CanArchiveSession)
+	/** Ask the user to confirm deleting - most obvious implementation is showing a dialog box */
+	SLATE_EVENT(FCanRemoveSession, CanDeleteSession)
 	
 	SLATE_END_ARGS();
 
@@ -186,6 +192,8 @@ private:
 	FSessionDelegate OnSessionClicked;
 	FSessionDelegate OnSessionDoubleClicked;
 	FSessionDelegate OnRequestedDeleteSession;
+	FCanRemoveSession CanArchiveSession;
+	FCanRemoveSession CanDeleteSession;
 
 	// The items displayed in the session list view. It might be filtered and sorted compared to the full list hold by the controller.
 	TArray<TSharedPtr<FConcertSessionItem>> Sessions;
diff --git a/Engine/Plugins/Developer/Concert/MultiUserClient/Source/MultiUserClient/Private/Widgets/SConcertBrowser.cpp b/Engine/Plugins/Developer/Concert/MultiUserClient/Source/MultiUserClient/Private/Widgets/SConcertBrowser.cpp
--- a/Engine/Plugins/Developer/Concert/MultiUserClient/Source/MultiUserClient/Private/Widgets/SConcertBrowser.cpp
+++ b/Engine/Plugins/Developer/Concert/MultiUserClient/Source/MultiUserClient/Private/Widgets/SConcertBrowser.cpp
@@ -1036,6 +1036,7 @@ private:
 	// Manipulates the sessions view (the array and the UI).
 	void OnSessionSelectionChanged(TSharedPtr<FConcertSessionItem> SelectedSession);
 	void OnSessionDoubleClicked(TSharedPtr<FConcertSessionItem> SelectedSession);
+	bool CanDeleteSession(TSharedPtr<FConcertSessionItem> DeletedItem);
 
 	// Update server/session/clients lists.
 	EActiveTimerReturnType TickDiscovery(double InCurrentTime, float InDeltaTime);
@@ -1713,6 +1714,15 @@ void SConcertClientSessionBrowser::OnSessionDoubleClicked(TSharedPtr<FConcertSes
 	}
 }
 
+bool SConcertClientSessionBrowser::CanDeleteSession(TSharedPtr<FConcertSessionItem> DeletedItem)
+{
+	const FText SessionNameInText = FText::FromString(DeletedItem->SessionName);
+	const FText SeverNameInText = FText::FromString(DeletedItem->ServerName);
+	const FText ConfirmationMessage = FText::Format(LOCTEXT("DeleteSessionConfirmationMessage", "Do you really want to delete the session \"{0}\" from the server \"{1}\"?"), SessionNameInText, SeverNameInText);
+	const FText ConfirmationTitle = LOCTEXT("DeleteSessionConfirmationTitle", "Delete Session Confirmation");
+	return FMessageDialog::Open(EAppMsgType::YesNo, ConfirmationMessage, &ConfirmationTitle) == EAppReturnType::Yes;
+}
+
 bool SConcertClientSessionBrowser::IsLaunchServerButtonEnabled() const
 {
 	return !bLocalServerRunning;
