Add missing files

diff --git a/Engine/Plugins/Developer/Concert/ConcertUI/ConcertServerUI/Source/ConcertServerUI/Private/Widgets/Session/ServerSessionHistoryController.cpp b/Engine/Plugins/Developer/Concert/ConcertUI/ConcertServerUI/Source/ConcertServerUI/Private/Widgets/Session/ServerSessionHistoryController.cpp
new file mode 100644
--- /dev/null
+++ b/Engine/Plugins/Developer/Concert/ConcertUI/ConcertServerUI/Source/ConcertServerUI/Private/Widgets/Session/ServerSessionHistoryController.cpp
@@ -0,0 +1,72 @@
+﻿// Copyright Epic Games, Inc. All Rights Reserved.
+
+#include "ServerSessionHistoryController.h"
+
+#include "ConcertSyncSessionDatabase.h"
+#include "IConcertSession.h"
+#include "IConcertSyncServer.h"
+
+FServerSessionHistoryController::FServerSessionHistoryController(TSharedRef<IConcertServerSession> InspectedSession, TSharedRef<IConcertSyncServer> SyncServer)
+	: InspectedSession(MoveTemp(InspectedSession))
+	, SyncServer(MoveTemp(SyncServer))
+{}
+
+void FServerSessionHistoryController::GetActivities(int64 MaximumNumberOfActivities, TMap<FGuid, FConcertClientInfo>& OutEndpointClientInfoMap, TArray<FConcertSessionActivity>& OutFetchedActivities) const
+{
+	if (const TOptional<FConcertSyncSessionDatabaseNonNullPtr> Database = SyncServer->GetLiveSessionDatabase(InspectedSession->GetId()))
+	{
+		OutEndpointClientInfoMap.Reset();
+		OutFetchedActivities.Reset();
+
+		int64 LastActivityId = INDEX_NONE;
+		Database->GetActivityMaxId(LastActivityId);
+		const int64 FirstActivityIdToFetch = FMath::Max<int64>(1, LastActivityId - MaximumNumberOfActivities);
+		Database->EnumerateActivitiesInRange(FirstActivityIdToFetch, MaximumNumberOfActivities, [this, &Database, &OutEndpointClientInfoMap, &OutFetchedActivities](FConcertSyncActivity&& InActivity)
+		{
+			if (!OutEndpointClientInfoMap.Contains(InActivity.EndpointId))
+			{
+				FConcertSyncEndpointData EndpointData;
+				if (Database->GetEndpoint(InActivity.EndpointId, EndpointData))
+				{
+					OutEndpointClientInfoMap.Add(InActivity.EndpointId, EndpointData.ClientInfo);
+				}
+			}
+
+			FStructOnScope ActivitySummary;
+			if (InActivity.EventSummary.GetPayload(ActivitySummary))
+			{
+				OutFetchedActivities.Emplace(MoveTemp(InActivity), MoveTemp(ActivitySummary));
+			}
+
+			return true;
+		});
+	}
+}
+
+bool FServerSessionHistoryController::GetPackageEvent(const FConcertSessionActivity& Activity, FConcertSyncPackageEventMetaData& OutPackageEvent) const
+{
+	if (const TOptional<FConcertSyncSessionDatabaseNonNullPtr> Database = SyncServer->GetLiveSessionDatabase(InspectedSession->GetId()))
+	{
+		return Database->GetPackageEventMetaData(Activity.Activity.EventId, OutPackageEvent.PackageRevision, OutPackageEvent.PackageInfo);
+	}
+	
+	return false;
+}
+
+TFuture<TOptional<FConcertSyncTransactionEvent>> FServerSessionHistoryController::GetTransactionEvent(const FConcertSessionActivity& Activity) const
+{
+	if (const TOptional<FConcertSyncSessionDatabaseNonNullPtr> Database = SyncServer->GetLiveSessionDatabase(InspectedSession->GetId()))
+	{
+		return FindOrRequestTransactionEvent(*Database, Activity.Activity.EventId);
+	}
+	
+	return MakeFulfilledPromise<TOptional<FConcertSyncTransactionEvent>>().GetFuture(); // Not found.
+}
+
+TFuture<TOptional<FConcertSyncTransactionEvent>> FServerSessionHistoryController::FindOrRequestTransactionEvent(const FConcertSyncSessionDatabase& Database, const int64 TransactionEventId) const
+{
+	FConcertSyncTransactionEvent TransactionEvent;
+	return Database.GetTransactionEvent(TransactionEventId, TransactionEvent, false)
+		?	MakeFulfilledPromise<TOptional<FConcertSyncTransactionEvent>>(MoveTemp(TransactionEvent)).GetFuture()
+		:	MakeFulfilledPromise<TOptional<FConcertSyncTransactionEvent>>().GetFuture();
+}
diff --git a/Engine/Plugins/Developer/Concert/ConcertUI/ConcertServerUI/Source/ConcertServerUI/Private/Widgets/Session/ServerSessionHistoryController.h b/Engine/Plugins/Developer/Concert/ConcertUI/ConcertServerUI/Source/ConcertServerUI/Private/Widgets/Session/ServerSessionHistoryController.h
new file mode 100644
--- /dev/null
+++ b/Engine/Plugins/Developer/Concert/ConcertUI/ConcertServerUI/Source/ConcertServerUI/Private/Widgets/Session/ServerSessionHistoryController.h
@@ -0,0 +1,32 @@
+﻿// Copyright Epic Games, Inc. All Rights Reserved.
+
+#pragma once
+
+#include "CoreMinimal.h"
+#include "SessionHistory/AbstractSessionHistoryController.h"
+
+class FConcertSyncSessionDatabase;
+class IConcertSyncServer;
+class IConcertServerSession;
+
+class FServerSessionHistoryController : public FAbstractSessionHistoryController
+{
+public:
+
+	FServerSessionHistoryController(TSharedRef<IConcertServerSession> InspectedSession, TSharedRef<IConcertSyncServer> SyncServer);
+	
+protected:
+
+	//~ Begin FAbstractSessionHistoryController Interface
+	virtual void GetActivities(int64 MaximumNumberOfActivities, TMap<FGuid, FConcertClientInfo>& OutEndpointClientInfoMap, TArray<FConcertSessionActivity>& OutFetchedActivities) const override;
+	virtual bool GetPackageEvent(const FConcertSessionActivity& Activity, FConcertSyncPackageEventMetaData& OutPackageEvent) const override;
+	virtual TFuture<TOptional<FConcertSyncTransactionEvent>> GetTransactionEvent(const FConcertSessionActivity& Activity) const override;
+	//~ End FAbstractSessionHistoryController Interface
+
+private:
+
+	TSharedRef<IConcertServerSession> InspectedSession;
+	TSharedRef<IConcertSyncServer> SyncServer;
+
+	TFuture<TOptional<FConcertSyncTransactionEvent>> FindOrRequestTransactionEvent(const FConcertSyncSessionDatabase& Database, const int64 TransactionEventId) const;
+};
