Fix missing and incorrect dialog message when deleting / archiving session in client and server

diff --git a/Engine/Plugins/Developer/Concert/ConcertUI/ConcertServerUI/Source/ConcertServerUI/Private/Widgets/Browser/SConcertServerSessionBrowser.cpp b/Engine/Plugins/Developer/Concert/ConcertUI/ConcertServerUI/Source/ConcertServerUI/Private/Widgets/Browser/SConcertServerSessionBrowser.cpp
--- a/Engine/Plugins/Developer/Concert/ConcertUI/ConcertServerUI/Source/ConcertServerUI/Private/Widgets/Browser/SConcertServerSessionBrowser.cpp
+++ b/Engine/Plugins/Developer/Concert/ConcertUI/ConcertServerUI/Source/ConcertServerUI/Private/Widgets/Browser/SConcertServerSessionBrowser.cpp
@@ -6,8 +6,8 @@
 #include "Widgets/Browser/ConcertServerSessionBrowserController.h"
 #include "Widgets/StatusBar/SConcertStatusBar.h"
 
-#include "SPositiveActionButton.h"
 #include "Dialog/SMessageDialog.h"
+#include "Framework/Docking/TabManager.h"
 #include "SessionBrowser/ConcertSessionItem.h"
 #include "Widgets/ConcertServerTabs.h"
 #include "Widgets/Input/SSearchBox.h"
@@ -16,6 +16,11 @@
 
 #define LOCTEXT_NAMESPACE "UnrealMultiUserUI"
 
+SConcertServerSessionBrowser::~SConcertServerSessionBrowser()
+{
+	UnregisterFromOnRootWindowClosed();
+}
+
 void SConcertServerSessionBrowser::Construct(const FArguments& InArgs, TSharedRef<FConcertServerSessionBrowserController> InController)
 {
 	Controller = InController;
@@ -51,44 +56,120 @@ TSharedRef<SWidget> SConcertServerSessionBrowser::MakeSessionTableView(const FAr
 	SearchText = MakeShared<FText>();
 	return SAssignNew(SessionBrowser, SConcertSessionBrowser, Controller.Pin().ToSharedRef(), SearchText)
 		.OnSessionDoubleClicked(InArgs._DoubleClickSession)
-		.CanArchiveSession(this, &SConcertServerSessionBrowser::ConfirmArchiveOperationWithDialog)
-		.CanDeleteSession(this, &SConcertServerSessionBrowser::ConfirmDeleteOperationWithDialog);
+		.OnRequestedDeleteSession(this, &SConcertServerSessionBrowser::RequestDeleteSession)
+		// Pretend a modal dialog said no - RequestDeleteSession will show non-modal dialog
+		.CanDeleteArchivedSession_Lambda([](TSharedPtr<FConcertSessionItem>) { return false; })
+		.CanDeleteActiveSession_Lambda([](TSharedPtr<FConcertSessionItem>) { return false; });
 }
 
-bool SConcertServerSessionBrowser::ConfirmArchiveOperationWithDialog(TSharedPtr<FConcertSessionItem> SessionItem)
+void SConcertServerSessionBrowser::RequestDeleteSession(TSharedPtr<FConcertSessionItem> SessionItem) 
+{
+	if (const TSharedPtr<SMessageDialog> PinnedDeleteSessionDialog = DeleteSessionDialog.Pin())
+	{
+		// Force user to answer previous question first
+		PinnedDeleteSessionDialog->BringToFront();
+		PinnedDeleteSessionDialog->FlashWindow();
+	}
+	else
+	{
+		// Modal dialog would cause server to halt thread and stop communicating with connected clients
+		switch (SessionItem->Type)
+		{
+		case FConcertSessionItem::EType::ActiveSession:
+			DeleteActiveSessionWithNonModalQuestion(SessionItem); break;
+		case FConcertSessionItem::EType::ArchivedSession:
+			DeleteArchivedSessionWithNonModalQuestion(SessionItem); break;
+		default:
+			break;
+		}
+
+		// The dialog should be closed when the root window is closed
+		if (const TSharedPtr<SWindow> RootWindow = FGlobalTabmanager::Get()->GetRootWindow();
+			RootWindow && !RootWindow->GetOnWindowClosedEvent().IsBoundToObject(this))
+		{
+			RootWindow->GetOnWindowClosedEvent().AddSP(this, &SConcertServerSessionBrowser::OnRootWindowClosed);
+		}
+	}
+}
+
+void SConcertServerSessionBrowser::OnRootWindowClosed(const TSharedRef<SWindow>&) const
+{
+	if (const TSharedPtr<SMessageDialog> PinnedDeleteSessionDialog = DeleteSessionDialog.Pin())
+	{
+		PinnedDeleteSessionDialog->RequestDestroyWindow();
+	}
+	
+	UnregisterFromOnRootWindowClosed();
+}
+
+void SConcertServerSessionBrowser::UnregisterFromOnRootWindowClosed() const
+{
+	if (const TSharedPtr<SWindow> RootWindow = FGlobalTabmanager::Get()->GetRootWindow())
+	{
+		RootWindow->GetOnWindowClosedEvent().RemoveAll(this);
+	}
+}
+
+void SConcertServerSessionBrowser::DeleteArchivedSessionWithNonModalQuestion(TSharedPtr<FConcertSessionItem> SessionItem)
 {
-	const int32 NumUsers = Controller.Pin()->GetNumConnectedClients(SessionItem->SessionId);
 	const FText Message = FText::Format(
-		LOCTEXT("ArchiveDescription", "There {0}|plural(one=is,other=are) {0} connected {0}|plural(one=client,other=clients) in the current session.\nArchiving a session will force all connected clients to disconnect."),
-		NumUsers
+	LOCTEXT("DeleteArchivedDescription", "Deleting a session will cause all associated data to be removed.\n\nDelete {0}?"),
+		FText::FromString(SessionItem->SessionName)
 		);
 
-	constexpr int32 ArchiveIndex = 0;
+	auto DeleteArchived = [WeakController = TWeakPtr<FConcertServerSessionBrowserController>(Controller), WeakSessionItem = TWeakPtr<FConcertSessionItem>(SessionItem)]()
+	{
+		const TSharedPtr<FConcertServerSessionBrowserController> PinnedController = WeakController.Pin();
+		const TSharedPtr<FConcertSessionItem> PinnedSessionItem = WeakSessionItem.Pin();
+		if (PinnedController && PinnedSessionItem)
+		{
+			PinnedController->DeleteArchivedSession(PinnedSessionItem->ServerAdminEndpointId, PinnedSessionItem->SessionId);
+		}
+	};
 	const TSharedRef<SMessageDialog> Dialog = SNew(SMessageDialog)
-		.Title(LOCTEXT("DisconnectUsersTitle", "Force Users to Disconnect?"))
+		.Title(LOCTEXT("DisconnectUsersTitle", "Delete session?"))
 		.IconBrush("Icons.WarningWithColor.Large")
 		.Message(Message)
 		.Buttons({
-			SMessageDialog::FButton(LOCTEXT("ArchiveButton", "Archive")).SetPrimary(true),
-			SMessageDialog::FButton(LOCTEXT("CancelButton", "Cancel"))
+			SMessageDialog::FButton(LOCTEXT("DeleteArchivedButton", "Delete"))
+				.SetOnClicked(FSimpleDelegate::CreateLambda(DeleteArchived)),
+			SMessageDialog::FButton(LOCTEXT("CancelButton", "Cancel")).SetPrimary(true)
 		});
-	return Dialog->ShowModal() == ArchiveIndex;
+	
+	DeleteSessionDialog = Dialog;
+	Dialog->Show();
 }
 
-bool SConcertServerSessionBrowser::ConfirmDeleteOperationWithDialog(TSharedPtr<FConcertSessionItem> SessionItem)
+void SConcertServerSessionBrowser::DeleteActiveSessionWithNonModalQuestion(TSharedPtr<FConcertSessionItem> SessionItem)
 {
-	const FText Message = LOCTEXT("DeleteDescription", "Deleting a session will cause all associated data to be removed.");
-	
-	constexpr int32 DeleteIndex = 0;
+	const int32 NumUsers = Controller.Pin()->GetNumConnectedClients(SessionItem->SessionId);
+	const FText Message = FText::Format(
+		LOCTEXT("DeletedActiveDescription", "There {0}|plural(one=is,other=are) {0} connected {0}|plural(one=client,other=clients) in the current session.\nDeleting a session will force all connected clients to disconnect.\n\nDelete {1}?"),
+		NumUsers,
+		FText::FromString(SessionItem->SessionName)
+		);
+
+	auto DeleteActive = [WeakController = TWeakPtr<FConcertServerSessionBrowserController>(Controller), WeakSessionItem = TWeakPtr<FConcertSessionItem>(SessionItem)]()
+	{
+		const TSharedPtr<FConcertServerSessionBrowserController> PinnedController = WeakController.Pin();
+		const TSharedPtr<FConcertSessionItem> PinnedSessionItem = WeakSessionItem.Pin();
+		if (PinnedController && PinnedSessionItem)
+		{
+			PinnedController->DeleteActiveSession(PinnedSessionItem->ServerAdminEndpointId, PinnedSessionItem->SessionId);
+		}
+	};
 	const TSharedRef<SMessageDialog> Dialog = SNew(SMessageDialog)
-		.Title(LOCTEXT("DisconnectUsersTitle", "Delete session?"))
+		.Title(LOCTEXT("DisconnectUsersTitle", "Force Users to Disconnect?"))
 		.IconBrush("Icons.WarningWithColor.Large")
 		.Message(Message)
 		.Buttons({
-			SMessageDialog::FButton(LOCTEXT("DeleteButton", "Delete")),
-			SMessageDialog::FButton(LOCTEXT("CancelButton", "Cancel")).SetPrimary(true)
+			SMessageDialog::FButton(LOCTEXT("DeleteActiveButton", "Delete")).SetPrimary(true)
+				.SetOnClicked(FSimpleDelegate::CreateLambda(DeleteActive)),
+			SMessageDialog::FButton(LOCTEXT("CancelButton", "Cancel"))
 		});
-	return Dialog->ShowModal() == DeleteIndex;
+	
+	DeleteSessionDialog = Dialog;
+	Dialog->Show();
 }
 
 #undef LOCTEXT_NAMESPACE
diff --git a/Engine/Plugins/Developer/Concert/ConcertUI/ConcertServerUI/Source/ConcertServerUI/Private/Widgets/Browser/SConcertServerSessionBrowser.h b/Engine/Plugins/Developer/Concert/ConcertUI/ConcertServerUI/Source/ConcertServerUI/Private/Widgets/Browser/SConcertServerSessionBrowser.h
--- a/Engine/Plugins/Developer/Concert/ConcertUI/ConcertServerUI/Source/ConcertServerUI/Private/Widgets/Browser/SConcertServerSessionBrowser.h
+++ b/Engine/Plugins/Developer/Concert/ConcertUI/ConcertServerUI/Source/ConcertServerUI/Private/Widgets/Browser/SConcertServerSessionBrowser.h
@@ -7,6 +7,7 @@
 #include "Widgets/DeclarativeSyntaxSupport.h"
 #include "Widgets/SCompoundWidget.h"
 
+class SMessageDialog;
 class FConcertServerSessionBrowserController;
 class SSearchBox;
 
@@ -18,6 +19,7 @@ public:
 	SLATE_BEGIN_ARGS(SConcertServerSessionBrowser) { }
 		SLATE_EVENT(FSessionDelegate, DoubleClickSession)
 	SLATE_END_ARGS()
+	~SConcertServerSessionBrowser();
 
 	void Construct(const FArguments& InArgs, TSharedRef<FConcertServerSessionBrowserController> InController);
 
@@ -27,12 +29,17 @@ private:
 
 	/** We can ask the controller about information and notify it about UI events. */
 	TWeakPtr<FConcertServerSessionBrowserController> Controller;
+	/** Tracks whether there is a dialog asking the user to delete a session. Used to avoid opening multiple. */
+	TWeakPtr<SMessageDialog> DeleteSessionDialog;
 
 	TSharedPtr<FText> SearchText;
 	TSharedPtr<SConcertSessionBrowser> SessionBrowser;
 	
 	TSharedRef<SWidget> MakeSessionTableView(const FArguments& InArgs);
 
-	bool ConfirmArchiveOperationWithDialog(TSharedPtr<FConcertSessionItem> SessionItem);
-	bool ConfirmDeleteOperationWithDialog(TSharedPtr<FConcertSessionItem> SessionItem);
+	void RequestDeleteSession(TSharedPtr<FConcertSessionItem> SessionItem);
+	void OnRootWindowClosed(const TSharedRef<SWindow>&) const;
+	void UnregisterFromOnRootWindowClosed() const;
+	void DeleteArchivedSessionWithNonModalQuestion(TSharedPtr<FConcertSessionItem> SessionItem);
+	void DeleteActiveSessionWithNonModalQuestion(TSharedPtr<FConcertSessionItem> SessionItem);
 };
diff --git a/Engine/Plugins/Developer/Concert/ConcertUI/ConcertSharedSlate/Source/ConcertSharedSlate/Private/SessionBrowser/SConcertSessionBrowser.cpp b/Engine/Plugins/Developer/Concert/ConcertUI/ConcertSharedSlate/Source/ConcertSharedSlate/Private/SessionBrowser/SConcertSessionBrowser.cpp
--- a/Engine/Plugins/Developer/Concert/ConcertUI/ConcertSharedSlate/Source/ConcertSharedSlate/Private/SessionBrowser/SConcertSessionBrowser.cpp
+++ b/Engine/Plugins/Developer/Concert/ConcertUI/ConcertSharedSlate/Source/ConcertSharedSlate/Private/SessionBrowser/SConcertSessionBrowser.cpp
@@ -51,8 +51,8 @@ void SConcertSessionBrowser::Construct(const FArguments& InArgs, TSharedRef<ICon
 	OnSessionClicked = InArgs._OnSessionClicked;
 	OnSessionDoubleClicked = InArgs._OnSessionDoubleClicked;
 	OnRequestedDeleteSession = InArgs._OnRequestedDeleteSession;
-	CanArchiveSession = InArgs._CanArchiveSession;
-	CanDeleteSession = InArgs._CanDeleteSession;
+	CanDeleteArchivedSession = InArgs._CanDeleteArchivedSession;
+	CanDeleteActiveSession = InArgs._CanDeleteActiveSession;
 
 	// Setup search filter.
 	SearchedText = InSearchText; // Reload a previous search text (in any). Useful to remember searched text between join/leave sessions, but not persistent if the tab is closed.
@@ -1024,27 +1024,22 @@ void SConcertSessionBrowser::RequestRenameSession(const TSharedPtr<FConcertSessi
 
 void SConcertSessionBrowser::RequestDeleteSession(const TSharedPtr<FConcertSessionItem>& DeletedItem)
 {
-	const FText SessionNameInText = FText::FromString(DeletedItem->SessionName);
-	const FText SeverNameInText = FText::FromString(DeletedItem->ServerName);
-	const FText ConfirmationMessage = FText::Format(LOCTEXT("DeleteSessionConfirmationMessage", "Do you really want to delete the session \"{0}\" from the server \"{1}\"?"), SessionNameInText, SeverNameInText);
-	const FText ConfirmationTitle = LOCTEXT("DeleteSessionConfirmationTitle", "Delete Session Confirmation");
-	
 	if (DeletedItem->Type == FConcertSessionItem::EType::ActiveSession)
 	{
-		if (!CanArchiveSession.IsBound() || CanArchiveSession.Execute(DeletedItem))
+		if (!CanDeleteActiveSession.IsBound() || CanDeleteActiveSession.Execute(DeletedItem))
 		{
 			GetController()->DeleteActiveSession(DeletedItem->ServerAdminEndpointId, DeletedItem->SessionId);
-			OnRequestedDeleteSession.ExecuteIfBound(DeletedItem);
 		}
 	}
 	else if (DeletedItem->Type == FConcertSessionItem::EType::ArchivedSession)
 	{
-		if (!CanDeleteSession.IsBound() || CanDeleteSession.Execute(DeletedItem))
+		if (!CanDeleteArchivedSession.IsBound() || CanDeleteArchivedSession.Execute(DeletedItem))
 		{
 			GetController()->DeleteArchivedSession(DeletedItem->ServerAdminEndpointId, DeletedItem->SessionId);
-			OnRequestedDeleteSession.ExecuteIfBound(DeletedItem);
 		}
 	}
+	
+	OnRequestedDeleteSession.ExecuteIfBound(DeletedItem);
 }
 
 #undef LOCTEXT_NAMESPACE
\ No newline at end of file
diff --git a/Engine/Plugins/Developer/Concert/ConcertUI/ConcertSharedSlate/Source/ConcertSharedSlate/Public/SessionBrowser/SConcertSessionBrowser.h b/Engine/Plugins/Developer/Concert/ConcertUI/ConcertSharedSlate/Source/ConcertSharedSlate/Public/SessionBrowser/SConcertSessionBrowser.h
--- a/Engine/Plugins/Developer/Concert/ConcertUI/ConcertSharedSlate/Source/ConcertSharedSlate/Public/SessionBrowser/SConcertSessionBrowser.h
+++ b/Engine/Plugins/Developer/Concert/ConcertUI/ConcertSharedSlate/Source/ConcertSharedSlate/Public/SessionBrowser/SConcertSessionBrowser.h
@@ -75,9 +75,9 @@ public:
 	SLATE_EVENT(FSessionDelegate, OnRequestedDeleteSession)
 
 	/** Ask the user to confirm archiving - most obvious implementation is showing a dialog box */
-	SLATE_EVENT(FCanRemoveSession, CanArchiveSession)
+	SLATE_EVENT(FCanRemoveSession, CanDeleteArchivedSession)
 	/** Ask the user to confirm deleting - most obvious implementation is showing a dialog box */
-	SLATE_EVENT(FCanRemoveSession, CanDeleteSession)
+	SLATE_EVENT(FCanRemoveSession, CanDeleteActiveSession)
 	
 	SLATE_END_ARGS();
 
@@ -192,8 +192,8 @@ private:
 	FSessionDelegate OnSessionClicked;
 	FSessionDelegate OnSessionDoubleClicked;
 	FSessionDelegate OnRequestedDeleteSession;
-	FCanRemoveSession CanArchiveSession;
-	FCanRemoveSession CanDeleteSession;
+	FCanRemoveSession CanDeleteArchivedSession;
+	FCanRemoveSession CanDeleteActiveSession;
 
 	// The items displayed in the session list view. It might be filtered and sorted compared to the full list hold by the controller.
 	TArray<TSharedPtr<FConcertSessionItem>> Sessions;
diff --git a/Engine/Plugins/Developer/Concert/MultiUserClient/Source/MultiUserClient/Private/Widgets/SConcertBrowser.cpp b/Engine/Plugins/Developer/Concert/MultiUserClient/Source/MultiUserClient/Private/Widgets/SConcertBrowser.cpp
--- a/Engine/Plugins/Developer/Concert/MultiUserClient/Source/MultiUserClient/Private/Widgets/SConcertBrowser.cpp
+++ b/Engine/Plugins/Developer/Concert/MultiUserClient/Source/MultiUserClient/Private/Widgets/SConcertBrowser.cpp
@@ -1037,7 +1037,7 @@ private:
 	// Manipulates the sessions view (the array and the UI).
 	void OnSessionSelectionChanged(TSharedPtr<FConcertSessionItem> SelectedSession);
 	void OnSessionDoubleClicked(TSharedPtr<FConcertSessionItem> SelectedSession);
-	bool CanDeleteSession(TSharedPtr<FConcertSessionItem> DeletedItem);
+	bool ConfirmDeleteSessionWithDialog(TSharedPtr<FConcertSessionItem> SessionItem);
 
 	// Update server/session/clients lists.
 	EActiveTimerReturnType TickDiscovery(double InCurrentTime, float InDeltaTime);
@@ -1175,6 +1175,8 @@ TSharedRef<SWidget> SConcertClientSessionBrowser::MakeBrowserContent(TSharedPtr<
 					.OnSessionClicked_Raw(this, &SConcertClientSessionBrowser::OnSessionSelectionChanged)
 					.OnSessionDoubleClicked_Raw(this, &SConcertClientSessionBrowser::OnSessionDoubleClicked)
 					.OnRequestedDeleteSession_Lambda([this](auto) { UpdateDiscovery(); /* Don't wait up to 1s, kick discovery right now */ })
+					.CanDeleteArchivedSession(this, &SConcertClientSessionBrowser::ConfirmDeleteSessionWithDialog)
+					.CanDeleteActiveSession(this, &SConcertClientSessionBrowser::ConfirmDeleteSessionWithDialog)
 			]
 
 			// Session details.
@@ -1715,10 +1717,10 @@ void SConcertClientSessionBrowser::OnSessionDoubleClicked(TSharedPtr<FConcertSes
 	}
 }
 
-bool SConcertClientSessionBrowser::CanDeleteSession(TSharedPtr<FConcertSessionItem> DeletedItem)
+bool SConcertClientSessionBrowser::ConfirmDeleteSessionWithDialog(TSharedPtr<FConcertSessionItem> SessionItem)
 {
-	const FText SessionNameInText = FText::FromString(DeletedItem->SessionName);
-	const FText SeverNameInText = FText::FromString(DeletedItem->ServerName);
+	const FText SessionNameInText = FText::FromString(SessionItem->SessionName);
+	const FText SeverNameInText = FText::FromString(SessionItem->ServerName);
 	const FText ConfirmationMessage = FText::Format(LOCTEXT("DeleteSessionConfirmationMessage", "Do you really want to delete the session \"{0}\" from the server \"{1}\"?"), SessionNameInText, SeverNameInText);
 	const FText ConfirmationTitle = LOCTEXT("DeleteSessionConfirmationTitle", "Delete Session Confirmation");
 	return FMessageDialog::Open(EAppMsgType::YesNo, ConfirmationMessage, &ConfirmationTitle) == EAppReturnType::Yes;
