Update SSessionHistory on server when new session is created by client

#jria UE-145713

diff --git a/Engine/Plugins/Developer/Concert/ConcertUI/ConcertServerUI/Source/ConcertServerUI/Private/Widgets/Browser/ConcertServerSessionBrowserController.cpp b/Engine/Plugins/Developer/Concert/ConcertUI/ConcertServerUI/Source/ConcertServerUI/Private/Widgets/Browser/ConcertServerSessionBrowserController.cpp
--- a/Engine/Plugins/Developer/Concert/ConcertUI/ConcertServerUI/Source/ConcertServerUI/Private/Widgets/Browser/ConcertServerSessionBrowserController.cpp
+++ b/Engine/Plugins/Developer/Concert/ConcertUI/ConcertServerUI/Source/ConcertServerUI/Private/Widgets/Browser/ConcertServerSessionBrowserController.cpp
@@ -2,6 +2,7 @@
 
 #include "ConcertServerSessionBrowserController.h"
 
+#include "ConcertServerEvents.h"
 #include "ConcertServerStyle.h"
 #include "IConcertServer.h"
 #include "IConcertSyncServer.h"
@@ -24,6 +25,14 @@ int32 FConcertServerSessionBrowserController::GetNumConnectedClients(const FGuid
 	return ensure(Session) ? Session->GetSessionClients().Num() : 0;
 }
 
+FConcertServerSessionBrowserController::~FConcertServerSessionBrowserController()
+{
+	ConcertServerEvents::OnLiveSessionCreated().RemoveAll(this);
+	ConcertServerEvents::OnLiveSessionDestroyed().RemoveAll(this);
+	ConcertServerEvents::OnArchivedSessionCreated().RemoveAll(this);
+	ConcertServerEvents::OnArchivedSessionDestroyed().RemoveAll(this);
+}
+
 void FConcertServerSessionBrowserController::Init(const FConcertComponentInitParams& Params)
 {
 	ServerInstance = Params.Server;
@@ -35,6 +44,11 @@ void FConcertServerSessionBrowserController::Init(const FConcertComponentInitPar
 		.SetDisplayName(LOCTEXT("SessionBrowserTabTitle", "Session Browser"))
 		.SetTooltipText(LOCTEXT("SessionBrowserTooltipText", "A section to browse, start, archive, and restore server sessions."))
 		.SetIcon(FSlateIcon(FConcertServerStyle::GetStyleSetName(), TEXT("Concert.MultiUser")));
+
+	ConcertServerEvents::OnLiveSessionCreated().AddSP(this, &FConcertServerSessionBrowserController::OnLiveSessionCreated);
+	ConcertServerEvents::OnLiveSessionDestroyed().AddSP(this, &FConcertServerSessionBrowserController::OnLiveSessionDestroyed);
+	ConcertServerEvents::OnArchivedSessionCreated().AddSP(this, &FConcertServerSessionBrowserController::OnArchivedSessionCreated);
+	ConcertServerEvents::OnArchivedSessionDestroyed().AddSP(this, &FConcertServerSessionBrowserController::OnArchivedSessionDestroyed);
 }
 
 TArray<FConcertServerInfo> FConcertServerSessionBrowserController::GetServers() const
@@ -140,6 +154,16 @@ TSharedRef<SDockTab> FConcertServerSessionBrowserController::SpawnSessionBrowser
 	return DockTab;
 }
 
+void FConcertServerSessionBrowserController::RefreshSessionList()
+{
+	if (ConcertBrowser.IsValid())
+	{
+		// 1. Accumulate several events in same tick so list is only updated once
+		// 2. Several events, e.g. OnLiveSessionCreated, are called before the session creation is visible to getter functions
+		ConcertBrowser->RequestRefreshListNextTick();
+	}
+}
+
 void FConcertServerSessionBrowserController::OpenSession(TSharedPtr<FConcertSessionItem> SessionItem)
 {
 	Owner.Pin()->OpenSessionTab(SessionItem->SessionId);
@@ -183,7 +207,7 @@ void FConcertServerSessionBrowserController::NotifyUserOfFinishedSessionAction(c
 
 	if (bSuccess)
 	{
-		ConcertBrowser->RefreshSessionList();
+		ConcertBrowser->RequestRefreshListNextTick();
 	}
 }
 
diff --git a/Engine/Plugins/Developer/Concert/ConcertUI/ConcertServerUI/Source/ConcertServerUI/Private/Widgets/Browser/ConcertServerSessionBrowserController.h b/Engine/Plugins/Developer/Concert/ConcertUI/ConcertServerUI/Source/ConcertServerUI/Private/Widgets/Browser/ConcertServerSessionBrowserController.h
--- a/Engine/Plugins/Developer/Concert/ConcertUI/ConcertServerUI/Source/ConcertServerUI/Private/Widgets/Browser/ConcertServerSessionBrowserController.h
+++ b/Engine/Plugins/Developer/Concert/ConcertUI/ConcertServerUI/Source/ConcertServerUI/Private/Widgets/Browser/ConcertServerSessionBrowserController.h
@@ -6,6 +6,8 @@
 #include "SessionBrowser/IConcertSessionBrowserController.h"
 #include "Widgets/IConcertComponent.h"
 
+class IConcertServerSession;
+class IConcertServer;
 class FConcertSessionItem;
 class FSpawnTabArgs;
 class SDockTab;
@@ -21,6 +23,8 @@ class FConcertServerSessionBrowserController
 public:
 
 	int32 GetNumConnectedClients(const FGuid& SessionId) const;
+
+	~FConcertServerSessionBrowserController();
 	
 	//~ Begin IConcertComponent Interface
 	virtual void Init(const FConcertComponentInitParams& Params) override;
@@ -59,6 +63,14 @@ private:
 	
 	TSharedRef<SDockTab> SpawnSessionBrowserTab(const FSpawnTabArgs& Args);
 
+	// Update view when session list changes
+	void OnLiveSessionCreated(bool, const IConcertServer&, TSharedRef<IConcertServerSession>) { RefreshSessionList(); }
+	void OnLiveSessionDestroyed(const IConcertServer&, TSharedRef<IConcertServerSession>) { RefreshSessionList(); }
+	void OnArchivedSessionCreated(bool, const IConcertServer&, const FString&, const FConcertSessionInfo&) { RefreshSessionList(); }
+	void OnArchivedSessionDestroyed(const IConcertServer&, const FGuid&) { RefreshSessionList(); }
+	
+	void RefreshSessionList();
+
 	// Session actions
 	void OpenSession(TSharedPtr<FConcertSessionItem> SessionItem);
 	void RenameSession(const FGuid& ServerAdminEndpointId, const FGuid& SessionId, const FString& NewName);
@@ -66,6 +78,5 @@ private:
 
 	// Notifications
 	void NotifyUserOfFinishedSessionAction(const bool bSuccess, const FText& Title) { NotifyUserOfFinishedSessionAction(bSuccess, Title, FText::GetEmpty()); }
-    /** Shared logic for displaying a notification after creating, archiving, restoring, renaming and deleting a session. */
 	void NotifyUserOfFinishedSessionAction(const bool bSuccess, const FText& Title, const FText& Details);
 };
diff --git a/Engine/Plugins/Developer/Concert/ConcertUI/ConcertServerUI/Source/ConcertServerUI/Private/Widgets/Browser/SConcertServerSessionBrowser.h b/Engine/Plugins/Developer/Concert/ConcertUI/ConcertServerUI/Source/ConcertServerUI/Private/Widgets/Browser/SConcertServerSessionBrowser.h
--- a/Engine/Plugins/Developer/Concert/ConcertUI/ConcertServerUI/Source/ConcertServerUI/Private/Widgets/Browser/SConcertServerSessionBrowser.h
+++ b/Engine/Plugins/Developer/Concert/ConcertUI/ConcertServerUI/Source/ConcertServerUI/Private/Widgets/Browser/SConcertServerSessionBrowser.h
@@ -23,7 +23,17 @@ public:
 
 	void Construct(const FArguments& InArgs, TSharedRef<FConcertServerSessionBrowserController> InController);
 
-	void RefreshSessionList() { SessionBrowser->RefreshSessionList(); }
+	void RequestRefreshListNextTick() { bRequestedRefresh = true; }
+	virtual void Tick(const FGeometry& AllottedGeometry, const double InCurrentTime, const float InDeltaTime) override
+	{
+		if (bRequestedRefresh)
+		{
+			SessionBrowser->RefreshSessionList();
+			bRequestedRefresh = false;
+		}
+
+		SCompoundWidget::Tick(AllottedGeometry, InCurrentTime, InDeltaTime);
+	}
 	
 private:
 
@@ -32,6 +42,8 @@ private:
 	/** Tracks whether there is a dialog asking the user to delete a session. Used to avoid opening multiple. */
 	TWeakPtr<SMessageDialog> DeleteSessionDialog;
 
+	bool bRequestedRefresh = false;
+
 	TSharedPtr<FText> SearchText;
 	TSharedPtr<SConcertSessionBrowser> SessionBrowser;
 	
