Fix non-unity build error.

diff --git a/Engine/Plugins/Developer/Concert/ConcertUI/ConcertServerUI/Source/ConcertServerUI/Private/Widgets/Tabs/Archived/SConcertArchivedSessionInspector.h b/Engine/Plugins/Developer/Concert/ConcertUI/ConcertServerUI/Source/ConcertServerUI/Private/Widgets/Tabs/Archived/SConcertArchivedSessionInspector.h
--- a/Engine/Plugins/Developer/Concert/ConcertUI/ConcertServerUI/Source/ConcertServerUI/Private/Widgets/Tabs/Archived/SConcertArchivedSessionInspector.h
+++ b/Engine/Plugins/Developer/Concert/ConcertUI/ConcertServerUI/Source/ConcertServerUI/Private/Widgets/Tabs/Archived/SConcertArchivedSessionInspector.h
@@ -4,6 +4,7 @@
 
 #include "CoreMinimal.h"
 #include "Widgets/DeclarativeSyntaxSupport.h"
+#include "Widgets/Docking/SDockTab.h"
 #include "Widgets/SCompoundWidget.h"
 
 /** Designed as content for a tab. Displays information about an archived session. */
