Move files into tab folder

diff --git a/Engine/Plugins/Developer/Concert/ConcertUI/ConcertServerUI/Source/ConcertServerUI/Private/Widgets/ConcertServerWindowController.cpp b/Engine/Plugins/Developer/Concert/ConcertUI/ConcertServerUI/Source/ConcertServerUI/Private/Widgets/ConcertServerWindowController.cpp
--- a/Engine/Plugins/Developer/Concert/ConcertUI/ConcertServerUI/Source/ConcertServerUI/Private/Widgets/ConcertServerWindowController.cpp
+++ b/Engine/Plugins/Developer/Concert/ConcertUI/ConcertServerUI/Source/ConcertServerUI/Private/Widgets/ConcertServerWindowController.cpp
@@ -2,10 +2,12 @@
 
 #include "ConcertServerWindowController.h"
 
+#include "Browser/ConcertServerSessionBrowserController.h"
 #include "ConcertServerTabs.h"
 #include "IConcertServer.h"
 #include "IConcertSyncServer.h"
-#include "Browser/ConcertServerSessionBrowserController.h"
+#include "SessionTabs/Archived/ArchivedConcertSessionTab.h"
+#include "SessionTabs/Live/LiveConcertSessionTab.h"
 
 #include "Framework/Application/SlateApplication.h"
 #include "Framework/Docking/TabManager.h"
@@ -14,9 +16,7 @@
 #include "HAL/PlatformApplicationMisc.h"
 #include "Misc/ConfigCacheIni.h"
 #include "OutputLog/Public/OutputLogModule.h"
-#include "Session/LiveConcertSessionTab.h"
 
-#include "Tabs/Archived/ArchivedConcertSessionTab.h"
 
 #include "Widgets/SWindow.h"
 
diff --git a/Engine/Plugins/Developer/Concert/ConcertUI/ConcertServerUI/Source/ConcertServerUI/Private/Widgets/Session/ConcertServerSessionTab.cpp b/Engine/Plugins/Developer/Concert/ConcertUI/ConcertServerUI/Source/ConcertServerUI/Private/Widgets/Session/ConcertServerSessionTab.cpp
deleted file mode 100644
diff --git a/Engine/Plugins/Developer/Concert/ConcertUI/ConcertServerUI/Source/ConcertServerUI/Private/Widgets/Session/ConcertServerSessionTab.h b/Engine/Plugins/Developer/Concert/ConcertUI/ConcertServerUI/Source/ConcertServerUI/Private/Widgets/Session/ConcertServerSessionTab.h
deleted file mode 100644
diff --git a/Engine/Plugins/Developer/Concert/ConcertUI/ConcertServerUI/Source/ConcertServerUI/Private/Widgets/Session/LiveConcertSessionTab.cpp b/Engine/Plugins/Developer/Concert/ConcertUI/ConcertServerUI/Source/ConcertServerUI/Private/Widgets/Session/LiveConcertSessionTab.cpp
deleted file mode 100644
diff --git a/Engine/Plugins/Developer/Concert/ConcertUI/ConcertServerUI/Source/ConcertServerUI/Private/Widgets/Session/LiveConcertSessionTab.h b/Engine/Plugins/Developer/Concert/ConcertUI/ConcertServerUI/Source/ConcertServerUI/Private/Widgets/Session/LiveConcertSessionTab.h
deleted file mode 100644
diff --git a/Engine/Plugins/Developer/Concert/ConcertUI/ConcertServerUI/Source/ConcertServerUI/Private/Widgets/Session/PackageViewer/ConcertSessionPackageViewerController.cpp b/Engine/Plugins/Developer/Concert/ConcertUI/ConcertServerUI/Source/ConcertServerUI/Private/Widgets/Session/PackageViewer/ConcertSessionPackageViewerController.cpp
deleted file mode 100644
diff --git a/Engine/Plugins/Developer/Concert/ConcertUI/ConcertServerUI/Source/ConcertServerUI/Private/Widgets/Session/PackageViewer/ConcertSessionPackageViewerController.h b/Engine/Plugins/Developer/Concert/ConcertUI/ConcertServerUI/Source/ConcertServerUI/Private/Widgets/Session/PackageViewer/ConcertSessionPackageViewerController.h
deleted file mode 100644
diff --git a/Engine/Plugins/Developer/Concert/ConcertUI/ConcertServerUI/Source/ConcertServerUI/Private/Widgets/Session/PackageViewer/SConcertSessionPackageViewer.cpp b/Engine/Plugins/Developer/Concert/ConcertUI/ConcertServerUI/Source/ConcertServerUI/Private/Widgets/Session/PackageViewer/SConcertSessionPackageViewer.cpp
deleted file mode 100644
diff --git a/Engine/Plugins/Developer/Concert/ConcertUI/ConcertServerUI/Source/ConcertServerUI/Private/Widgets/Session/PackageViewer/SConcertSessionPackageViewer.h b/Engine/Plugins/Developer/Concert/ConcertUI/ConcertServerUI/Source/ConcertServerUI/Private/Widgets/Session/PackageViewer/SConcertSessionPackageViewer.h
deleted file mode 100644
diff --git a/Engine/Plugins/Developer/Concert/ConcertUI/ConcertServerUI/Source/ConcertServerUI/Private/Widgets/Session/SConcertSessionInspector.cpp b/Engine/Plugins/Developer/Concert/ConcertUI/ConcertServerUI/Source/ConcertServerUI/Private/Widgets/Session/SConcertSessionInspector.cpp
deleted file mode 100644
diff --git a/Engine/Plugins/Developer/Concert/ConcertUI/ConcertServerUI/Source/ConcertServerUI/Private/Widgets/Session/SConcertSessionInspector.h b/Engine/Plugins/Developer/Concert/ConcertUI/ConcertServerUI/Source/ConcertServerUI/Private/Widgets/Session/SConcertSessionInspector.h
deleted file mode 100644
diff --git a/Engine/Plugins/Developer/Concert/ConcertUI/ConcertServerUI/Source/ConcertServerUI/Private/Widgets/Session/ServerSessionHistoryController.cpp b/Engine/Plugins/Developer/Concert/ConcertUI/ConcertServerUI/Source/ConcertServerUI/Private/Widgets/Session/ServerSessionHistoryController.cpp
deleted file mode 100644
diff --git a/Engine/Plugins/Developer/Concert/ConcertUI/ConcertServerUI/Source/ConcertServerUI/Private/Widgets/Session/ServerSessionHistoryController.h b/Engine/Plugins/Developer/Concert/ConcertUI/ConcertServerUI/Source/ConcertServerUI/Private/Widgets/Session/ServerSessionHistoryController.h
deleted file mode 100644
diff --git a/Engine/Plugins/Developer/Concert/ConcertUI/ConcertServerUI/Source/ConcertServerUI/Private/Widgets/SessionTabs/AbstractConcertSessionTab.cpp b/Engine/Plugins/Developer/Concert/ConcertUI/ConcertServerUI/Source/ConcertServerUI/Private/Widgets/SessionTabs/AbstractConcertSessionTab.cpp
new file mode 100644
--- /dev/null
+++ b/Engine/Plugins/Developer/Concert/ConcertUI/ConcertServerUI/Source/ConcertServerUI/Private/Widgets/SessionTabs/AbstractConcertSessionTab.cpp
@@ -0,0 +1,65 @@
+﻿// Copyright Epic Games, Inc. All Rights Reserved.
+
+#include "AbstractConcertSessionTab.h"
+
+#include "IConcertServer.h"
+#include "IConcertSyncServer.h"
+#include "Widgets/ConcertServerTabs.h"
+
+#include "Framework/Docking/TabManager.h"
+#include "Widgets/Docking/SDockTab.h"
+
+namespace UE::ConcertServerUI::Private
+{
+	static TOptional<FString> GetSessionName(const TSharedRef<IConcertSyncServer>& SyncServer, const FGuid SessionId)
+	{
+		if (const TSharedPtr<IConcertServerSession> LiveSessionInfo = SyncServer->GetConcertServer()->GetLiveSession(SessionId))
+		{
+			return LiveSessionInfo->GetSessionInfo().SessionName;
+		}
+		if (const TOptional<FConcertSessionInfo> ArchivedSessionInfo = SyncServer->GetConcertServer()->GetArchivedSessionInfo(SessionId))
+		{
+			return ArchivedSessionInfo->SessionName;
+		}
+		return {};
+	}
+}
+
+FAbstractConcertSessionTab::FAbstractConcertSessionTab(TSharedRef<IConcertSyncServer> SyncServer)
+	: SyncServer(MoveTemp(SyncServer))
+{}
+
+void FAbstractConcertSessionTab::OpenSessionTab()
+{
+	const TSharedRef<FGlobalTabmanager>& TabManager = FGlobalTabmanager::Get();
+	const FTabId TabId { *GetTabId() };
+
+	EnsureInitDockTab();
+	if (TabManager->FindExistingLiveTab(TabId))
+	{
+		TabManager->DrawAttention(DockTab.ToSharedRef());
+	}
+	else
+	{
+		const FTabManager::FLastMajorOrNomadTab Search(ConcertServerTabs::GetSessionBrowserTabId());
+		FGlobalTabmanager::Get()->InsertNewDocumentTab(*GetTabId(), Search, DockTab.ToSharedRef());
+
+		OnOpenTab();
+	}
+}
+
+void FAbstractConcertSessionTab::EnsureInitDockTab()
+{
+	if (!DockTab.IsValid())
+	{
+		const TOptional<FString> SessionName = UE::ConcertServerUI::Private::GetSessionName(SyncServer, GetSessionID());
+		check(SessionName);
+		
+		const FText Title = FText::FromString(*SessionName);
+		DockTab = SNew(SDockTab)
+			.Label(Title)
+			.TabRole(MajorTab);
+		
+		CreateDockContent(DockTab.ToSharedRef());
+	}
+}
diff --git a/Engine/Plugins/Developer/Concert/ConcertUI/ConcertServerUI/Source/ConcertServerUI/Private/Widgets/SessionTabs/AbstractConcertSessionTab.h b/Engine/Plugins/Developer/Concert/ConcertUI/ConcertServerUI/Source/ConcertServerUI/Private/Widgets/SessionTabs/AbstractConcertSessionTab.h
new file mode 100644
--- /dev/null
+++ b/Engine/Plugins/Developer/Concert/ConcertUI/ConcertServerUI/Source/ConcertServerUI/Private/Widgets/SessionTabs/AbstractConcertSessionTab.h
@@ -0,0 +1,42 @@
+﻿// Copyright Epic Games, Inc. All Rights Reserved.
+
+#pragma once
+
+#include "CoreMinimal.h"
+
+class IConcertSyncServer;
+class SDockTab;
+class SWindow;
+
+/** Shared functionality for a tab that manages a session */
+class FAbstractConcertSessionTab
+{
+public:
+
+    FAbstractConcertSessionTab(TSharedRef<IConcertSyncServer> SyncServer);
+	virtual ~FAbstractConcertSessionTab() = default;
+	
+	/** Opens or draws attention to the given tab */
+	void OpenSessionTab();
+	
+protected:
+
+	virtual FGuid GetSessionID() const = 0;
+	virtual void CreateDockContent(const TSharedRef<SDockTab>& DockTab) = 0;
+	virtual void OnOpenTab() = 0;
+	
+	/** Generates a tab ID for FTabManager::InsertNewDocumentTab */
+	FString GetTabId() const { return GetSessionID().ToString(); }
+	const TSharedRef<IConcertSyncServer>& GetSyncServer() const { return SyncServer; }
+	
+private:
+
+	/** Used to look up session name */
+	const TSharedRef<IConcertSyncServer> SyncServer;
+	
+	/** The tab containing the UI for InspectedSession */
+	TSharedPtr<SDockTab> DockTab;
+
+	/** Inits DockTab if it has not yet been inited */
+	void EnsureInitDockTab();
+};
diff --git a/Engine/Plugins/Developer/Concert/ConcertUI/ConcertServerUI/Source/ConcertServerUI/Private/Widgets/SessionTabs/Archived/ArchivedConcertSessionTab.cpp b/Engine/Plugins/Developer/Concert/ConcertUI/ConcertServerUI/Source/ConcertServerUI/Private/Widgets/SessionTabs/Archived/ArchivedConcertSessionTab.cpp
new file mode 100644
--- /dev/null
+++ b/Engine/Plugins/Developer/Concert/ConcertUI/ConcertServerUI/Source/ConcertServerUI/Private/Widgets/SessionTabs/Archived/ArchivedConcertSessionTab.cpp
@@ -0,0 +1,31 @@
+﻿// Copyright Epic Games, Inc. All Rights Reserved.
+
+#include "Widgets/SessionTabs/Archived/ArchivedConcertSessionTab.h"
+
+#include "Widgets/SessionTabs/Archived/SConcertArchivedSessionInspector.h"
+#include "Widgets/StatusBar/SConcertStatusBar.h"
+
+FArchivedConcertSessionTab::FArchivedConcertSessionTab(const FGuid& InspectedSessionID, TSharedRef<IConcertSyncServer> SyncServer, TAttribute<TSharedRef<SWindow>> ConstructUnderWindow)
+	: FAbstractConcertSessionTab(SyncServer)
+	, InspectedSessionID(InspectedSessionID)
+	, ConstructUnderWindow(ConstructUnderWindow)
+{}
+
+FGuid FArchivedConcertSessionTab::GetSessionID() const
+{
+	return InspectedSessionID;
+}
+
+void FArchivedConcertSessionTab::CreateDockContent(const TSharedRef<SDockTab>& InDockTab)
+{
+	InDockTab->SetContent(
+		SNew(SConcertArchivedSessionInspector, SConcertArchivedSessionInspector::FRequiredArgs(InDockTab, ConstructUnderWindow.Get()))
+			.StatusBar()
+			[
+				SNew(SConcertStatusBar, *GetTabId())
+			]
+		);
+}
+
+void FArchivedConcertSessionTab::OnOpenTab()
+{}
diff --git a/Engine/Plugins/Developer/Concert/ConcertUI/ConcertServerUI/Source/ConcertServerUI/Private/Widgets/SessionTabs/Archived/ArchivedConcertSessionTab.h b/Engine/Plugins/Developer/Concert/ConcertUI/ConcertServerUI/Source/ConcertServerUI/Private/Widgets/SessionTabs/Archived/ArchivedConcertSessionTab.h
new file mode 100644
--- /dev/null
+++ b/Engine/Plugins/Developer/Concert/ConcertUI/ConcertServerUI/Source/ConcertServerUI/Private/Widgets/SessionTabs/Archived/ArchivedConcertSessionTab.h
@@ -0,0 +1,35 @@
+﻿// Copyright Epic Games, Inc. All Rights Reserved.
+
+#pragma once
+
+#include "CoreMinimal.h"
+#include "SConcertArchivedSessionInspector.h"
+
+#include "Widgets/SessionTabs/AbstractConcertSessionTab.h"
+
+/** Manages the tab for an archived session.  */
+class FArchivedConcertSessionTab : public FAbstractConcertSessionTab
+{
+public:
+
+	FArchivedConcertSessionTab(const FGuid& InspectedSessionID, TSharedRef<IConcertSyncServer> SyncServer, TAttribute<TSharedRef<SWindow>> ConstructUnderWindow);
+
+protected:
+
+	//~ Begin FAbstractConcertSessionTab Interface
+	virtual FGuid GetSessionID() const override;
+	virtual void CreateDockContent(const TSharedRef<SDockTab>& InDockTab) override;
+	virtual void OnOpenTab() override;
+	//~ End FAbstractConcertSessionTab Interface
+	
+private:
+
+	/** The inspected session's ID */
+	FGuid InspectedSessionID;
+
+	/** Used later to obtain the window into which to add the tab */
+	TAttribute<TSharedRef<SWindow>> ConstructUnderWindow;
+
+	/** Displays session */
+	TSharedPtr<SConcertArchivedSessionInspector> Inspector;
+};
diff --git a/Engine/Plugins/Developer/Concert/ConcertUI/ConcertServerUI/Source/ConcertServerUI/Private/Widgets/SessionTabs/Archived/SConcertArchivedSessionInspector.cpp b/Engine/Plugins/Developer/Concert/ConcertUI/ConcertServerUI/Source/ConcertServerUI/Private/Widgets/SessionTabs/Archived/SConcertArchivedSessionInspector.cpp
new file mode 100644
--- /dev/null
+++ b/Engine/Plugins/Developer/Concert/ConcertUI/ConcertServerUI/Source/ConcertServerUI/Private/Widgets/SessionTabs/Archived/SConcertArchivedSessionInspector.cpp
@@ -0,0 +1,89 @@
+﻿// Copyright Epic Games, Inc. All Rights Reserved.
+
+#include "SConcertArchivedSessionInspector.h"
+
+#include "Framework/Docking/TabManager.h"
+#include "Session/History/SSessionHistory.h"
+#include "Widgets/Docking/SDockTab.h"
+#include "Widgets/Layout/SBorder.h"
+#include "Widgets/SBoxPanel.h"
+#include "Widgets/SNullWidget.h"
+
+#define LOCTEXT_NAMESPACE "SConcertArchivedSessionInspector"
+
+const FName SConcertArchivedSessionInspector::HistoryTabId("HistoryTabId");
+
+void SConcertArchivedSessionInspector::Construct(const FArguments& InArgs, const FRequiredArgs& InRequiredArgs)
+{
+	ChildSlot
+	[
+		SNew(SBorder)
+		.BorderImage(FAppStyle::Get().GetBrush("ToolPanel.GroupBorder"))
+		.Padding(FMargin(1.0f, 2.0f))
+		[
+			SNew(SVerticalBox)
+
+			// Toolbar
+			+SVerticalBox::Slot()
+			.AutoHeight()
+			[
+				SNullWidget::NullWidget
+			]
+
+			// Content
+			+SVerticalBox::Slot()
+			.FillHeight(1.f)
+			[
+				SNew(SBorder)
+				.BorderImage(FAppStyle::Get().GetBrush("Brushes.Title"))
+				.Padding(FMargin(0.f, 0.f, 0.f, 5.0f)) // Visually separate status bar from tabs
+				[
+					CreateTabs(InRequiredArgs)
+				]
+			]
+
+			// Status bar
+			+SVerticalBox::Slot()
+			.AutoHeight()
+			.HAlign(HAlign_Fill)
+			[
+				InArgs._StatusBar.Widget
+			]
+		]
+	];
+}
+
+TSharedRef<SWidget> SConcertArchivedSessionInspector::CreateTabs(const FRequiredArgs& RequiredArgs)
+{
+	TabManager = FGlobalTabmanager::Get()->NewTabManager(RequiredArgs.ConstructUnderMajorTab);
+	
+	TabManager->RegisterTabSpawner(HistoryTabId, FOnSpawnTab::CreateSP(this, &SConcertArchivedSessionInspector::SpawnActivityHistory))
+		.SetDisplayName(LOCTEXT("ActivityHistoryLabel", "History"));
+
+	const TSharedRef<FTabManager::FLayout> Layout = FTabManager::NewLayout("ConcertArchivedSessionLayout_v0.1")
+		->AddArea
+		(
+			FTabManager::NewPrimaryArea()
+				->SetOrientation(Orient_Horizontal)
+				->Split
+				(
+					FTabManager::NewStack()
+					->SetSizeCoefficient(0.5f)
+					->AddTab(HistoryTabId, ETabState::OpenedTab)
+				)
+		);
+
+	return TabManager->RestoreFrom(Layout, RequiredArgs.ConstructUnderWindow).ToSharedRef();
+}
+
+TSharedRef<SDockTab> SConcertArchivedSessionInspector::SpawnActivityHistory(const FSpawnTabArgs& Args)
+{
+	return SNew(SDockTab)
+		.Label(LOCTEXT("ActivityHistoryLabel", "History"))
+		.TabRole(PanelTab)
+		[
+			SNullWidget::NullWidget
+		]; 
+}
+
+#undef LOCTEXT_NAMESPACE
diff --git a/Engine/Plugins/Developer/Concert/ConcertUI/ConcertServerUI/Source/ConcertServerUI/Private/Widgets/SessionTabs/Archived/SConcertArchivedSessionInspector.h b/Engine/Plugins/Developer/Concert/ConcertUI/ConcertServerUI/Source/ConcertServerUI/Private/Widgets/SessionTabs/Archived/SConcertArchivedSessionInspector.h
new file mode 100644
--- /dev/null
+++ b/Engine/Plugins/Developer/Concert/ConcertUI/ConcertServerUI/Source/ConcertServerUI/Private/Widgets/SessionTabs/Archived/SConcertArchivedSessionInspector.h
@@ -0,0 +1,42 @@
+﻿// Copyright Epic Games, Inc. All Rights Reserved.
+
+#pragma once
+
+#include "CoreMinimal.h"
+#include "Widgets/DeclarativeSyntaxSupport.h"
+#include "Widgets/Docking/SDockTab.h"
+#include "Widgets/SCompoundWidget.h"
+
+/** Designed as content for a tab. Displays information about an archived session. */
+class SConcertArchivedSessionInspector : public SCompoundWidget
+{
+public:
+
+	struct FRequiredArgs
+	{
+		TSharedRef<SDockTab> ConstructUnderMajorTab;
+		TSharedRef<SWindow> ConstructUnderWindow;
+
+		FRequiredArgs(TSharedRef<SDockTab> ConstructUnderMajorTab, TSharedRef<SWindow> ConstructUnderWindow)
+			: ConstructUnderMajorTab(MoveTemp(ConstructUnderMajorTab))
+			, ConstructUnderWindow(MoveTemp(ConstructUnderWindow))
+		{}
+	};
+	
+	static const FName HistoryTabId;
+
+	SLATE_BEGIN_ARGS(SConcertArchivedSessionInspector)
+	{}
+		SLATE_NAMED_SLOT(FArguments, StatusBar)
+	SLATE_END_ARGS()
+
+	void Construct(const FArguments& InArgs, const FRequiredArgs& InRequiredArgs);
+
+private:
+	
+	/** Holds the tab manager that manages the front-end's tabs. */
+	TSharedPtr<FTabManager> TabManager;
+	
+	TSharedRef<SWidget> CreateTabs(const FRequiredArgs& RequiredArgs);
+	TSharedRef<SDockTab> SpawnActivityHistory(const FSpawnTabArgs& Args);
+};
diff --git a/Engine/Plugins/Developer/Concert/ConcertUI/ConcertServerUI/Source/ConcertServerUI/Private/Widgets/SessionTabs/Live/LiveConcertSessionTab.cpp b/Engine/Plugins/Developer/Concert/ConcertUI/ConcertServerUI/Source/ConcertServerUI/Private/Widgets/SessionTabs/Live/LiveConcertSessionTab.cpp
new file mode 100644
--- /dev/null
+++ b/Engine/Plugins/Developer/Concert/ConcertUI/ConcertServerUI/Source/ConcertServerUI/Private/Widgets/SessionTabs/Live/LiveConcertSessionTab.cpp
@@ -0,0 +1,48 @@
+﻿// Copyright Epic Games, Inc. All Rights Reserved.
+
+#include "LiveConcertSessionTab.h"
+
+#include "IConcertSession.h"
+#include "PackageViewer/ConcertSessionPackageViewerController.h"
+#include "SConcertSessionInspector.h"
+#include "Widgets/StatusBar/SConcertStatusBar.h"
+
+#include "Widgets/Docking/SDockTab.h"
+
+FLiveConcertSessionTab::FLiveConcertSessionTab(TSharedRef<IConcertServerSession> InspectedSession, TSharedRef<IConcertSyncServer> SyncServer, TAttribute<TSharedRef<SWindow>> ConstructUnderWindow)
+	: FAbstractConcertSessionTab(SyncServer)
+	, InspectedSession(MoveTemp(InspectedSession))
+	, ConstructUnderWindow(MoveTemp(ConstructUnderWindow))
+	, SessionHistoryController(MakeShared<FServerSessionHistoryController>(InspectedSession, SyncServer))
+	, PackageViewerController(MakeShared<FConcertSessionPackageViewerController>(InspectedSession, SyncServer))
+{}
+
+FGuid FLiveConcertSessionTab::GetSessionID() const
+{
+	return InspectedSession->GetSessionInfo().SessionId;
+}
+
+void FLiveConcertSessionTab::CreateDockContent(const TSharedRef<SDockTab>& InDockTab)
+{
+	const SConcertSessionInspector::FRequiredArgs WidgetArgs
+	{
+		InDockTab,
+		ConstructUnderWindow.Get(),
+		SessionHistoryController->GetSessionHistory(),
+		PackageViewerController->GetPackageViewer()
+	};
+	InDockTab->SetContent(
+		SNew(SConcertSessionInspector, WidgetArgs)
+			.StatusBar()
+			[
+				SNew(SConcertStatusBar, *GetTabId())
+			]
+		);
+}
+
+void FLiveConcertSessionTab::OnOpenTab()
+{
+	SessionHistoryController->ReloadActivities();
+	PackageViewerController->ReloadActivities();
+}
+
diff --git a/Engine/Plugins/Developer/Concert/ConcertUI/ConcertServerUI/Source/ConcertServerUI/Private/Widgets/SessionTabs/Live/LiveConcertSessionTab.h b/Engine/Plugins/Developer/Concert/ConcertUI/ConcertServerUI/Source/ConcertServerUI/Private/Widgets/SessionTabs/Live/LiveConcertSessionTab.h
new file mode 100644
--- /dev/null
+++ b/Engine/Plugins/Developer/Concert/ConcertUI/ConcertServerUI/Source/ConcertServerUI/Private/Widgets/SessionTabs/Live/LiveConcertSessionTab.h
@@ -0,0 +1,46 @@
+﻿// Copyright Epic Games, Inc. All Rights Reserved.
+
+#pragma once
+
+#include "CoreMinimal.h"
+#include "ServerSessionHistoryController.h"
+#include "Widgets/SessionTabs/AbstractConcertSessionTab.h"
+
+class FConcertSessionPackageViewerController;
+class IConcertSyncServer;
+class FServerSessionHistoryController;
+class IConcertServerSession;
+class SDockTab;
+class SWindow;
+
+/**
+ * Manages the tab that contains the UI for a session.
+ * It has access to the controllers and views needed for displaying a session.
+ */
+class FLiveConcertSessionTab : public FAbstractConcertSessionTab
+{
+public:
+
+	FLiveConcertSessionTab(TSharedRef<IConcertServerSession> InspectedSession, TSharedRef<IConcertSyncServer> SyncServer, TAttribute<TSharedRef<SWindow>> ConstructUnderWindow);
+
+protected:
+
+	//~ Begin FAbstractConcertSessionTab Interface
+	virtual FGuid GetSessionID() const override;
+	virtual void CreateDockContent(const TSharedRef<SDockTab>& InDockTab) override;
+	virtual void OnOpenTab() override;
+	//~ End FAbstractConcertSessionTab Interface
+	
+private:
+	
+	/** The session being inspected */
+	const TSharedRef<IConcertServerSession> InspectedSession;
+
+	/** Used to determine owning window during construction */
+	TAttribute<TSharedRef<SWindow>> ConstructUnderWindow;
+
+	/** Manages the session history widget */
+	const TSharedRef<FServerSessionHistoryController> SessionHistoryController;
+	/** Manages the package viewer widget */
+	const TSharedRef<FConcertSessionPackageViewerController> PackageViewerController;
+};
diff --git a/Engine/Plugins/Developer/Concert/ConcertUI/ConcertServerUI/Source/ConcertServerUI/Private/Widgets/SessionTabs/Live/PackageViewer/ConcertSessionPackageViewerController.cpp b/Engine/Plugins/Developer/Concert/ConcertUI/ConcertServerUI/Source/ConcertServerUI/Private/Widgets/SessionTabs/Live/PackageViewer/ConcertSessionPackageViewerController.cpp
new file mode 100644
--- /dev/null
+++ b/Engine/Plugins/Developer/Concert/ConcertUI/ConcertServerUI/Source/ConcertServerUI/Private/Widgets/SessionTabs/Live/PackageViewer/ConcertSessionPackageViewerController.cpp
@@ -0,0 +1,104 @@
+﻿// Copyright Epic Games, Inc. All Rights Reserved.
+
+#include "ConcertSessionPackageViewerController.h"
+
+#include "ConcertSyncSessionDatabase.h"
+#include "IConcertSyncServer.h"
+#include "SConcertSessionPackageViewer.h"
+#include "Concert/Private/ConcertServerSession.h"
+
+FConcertSessionPackageViewerController::FConcertSessionPackageViewerController(TSharedRef<IConcertServerSession> InspectedSession, TSharedRef<IConcertSyncServer> SyncServer)
+	: InspectedSession(MoveTemp(InspectedSession))
+	, SyncServer(MoveTemp(SyncServer))
+	, PackageViewer(MakePackageViewer())
+{
+	ReloadActivities();
+	
+	if (const TOptional<FConcertSyncSessionDatabaseNonNullPtr> Database = SyncServer->GetLiveSessionDatabase(InspectedSession->GetId()))
+	{
+		Database->OnActivityProduced().AddRaw(this, &FConcertSessionPackageViewerController::OnSessionProduced);
+	}
+}
+
+FConcertSessionPackageViewerController::~FConcertSessionPackageViewerController()
+{
+	if (const TOptional<FConcertSyncSessionDatabaseNonNullPtr> Database = SyncServer->GetLiveSessionDatabase(InspectedSession->GetId()))
+	{
+		Database->OnActivityProduced().RemoveAll(this);
+	}
+}
+
+void FConcertSessionPackageViewerController::ReloadActivities() const
+{
+	PackageViewer->ResetActivityList();
+	
+	if (const TOptional<FConcertSyncSessionDatabaseNonNullPtr> Database = SyncServer->GetLiveSessionDatabase(InspectedSession->GetId()))
+	{
+		struct FPackageActivity
+		{
+			FConcertSessionActivity Activity;
+			FConcertSyncPackageEventData PackageData;
+		};
+		
+		TMap<FName, FPackageActivity> LatestPackageActivities;
+		Database->EnumeratePackageActivities([this, &LatestPackageActivities](FConcertSyncActivity&& BasePart, FConcertSyncPackageEventData& EventData)
+		{
+			FStructOnScope ActivitySummary;
+			if (BasePart.EventSummary.GetPayload(ActivitySummary))
+			{
+				FPackageActivity Activity{ { BasePart, ActivitySummary }, EventData };
+				const bool bWasRenamed = Activity.PackageData.MetaData.PackageInfo.PackageUpdateType == EConcertPackageUpdateType::Renamed;
+				if (bWasRenamed)
+				{
+					LatestPackageActivities.Remove(EventData.MetaData.PackageInfo.PackageName);
+					LatestPackageActivities.Emplace(EventData.MetaData.PackageInfo.NewPackageName, MoveTemp(Activity));
+				}
+				else
+				{
+					LatestPackageActivities.Emplace(EventData.MetaData.PackageInfo.PackageName, MoveTemp(Activity));
+				}
+			}
+			return EBreakBehavior::Continue;
+		});
+
+		for (auto ActivityIt = LatestPackageActivities.CreateIterator(); ActivityIt; ++ActivityIt)
+		{
+			PackageViewer->AppendActivity(MoveTemp(ActivityIt->Value.Activity));
+		}
+	}
+}
+
+TSharedRef<SConcertSessionPackageViewer> FConcertSessionPackageViewerController::MakePackageViewer() const
+{
+	return SNew(SConcertSessionPackageViewer)
+		.GetClientInfo_Raw(this, &FConcertSessionPackageViewerController::GetClientInfo)
+		.GetPackageEvent_Raw(this, &FConcertSessionPackageViewerController::GetPackageEvent);
+}
+
+TOptional<FConcertClientInfo> FConcertSessionPackageViewerController::GetClientInfo(FGuid ClientId) const
+{
+	const TOptional<FConcertSyncSessionDatabaseNonNullPtr> Database = SyncServer->GetLiveSessionDatabase(InspectedSession->GetId());
+	FConcertSyncEndpointData Result;
+	if (Database && Database->GetEndpoint(ClientId, Result))
+	{
+		return Result.ClientInfo;
+	}
+	return {};
+}
+
+bool FConcertSessionPackageViewerController::GetPackageEvent(const FConcertSessionActivity& Activity, FConcertSyncPackageEventMetaData& OutPackageEvent) const
+{
+	if (const TOptional<FConcertSyncSessionDatabaseNonNullPtr> Database = SyncServer->GetLiveSessionDatabase(InspectedSession->GetId()))
+	{
+		return Database->GetPackageEventMetaData(Activity.Activity.EventId, OutPackageEvent.PackageRevision, OutPackageEvent.PackageInfo);
+	}
+	return false;
+}
+
+void FConcertSessionPackageViewerController::OnSessionProduced(const FConcertSyncActivity& ProducedActivity) const
+{
+	if (ProducedActivity.EventType == EConcertSyncActivityEventType::Package)
+	{
+		ReloadActivities();
+	}
+}
diff --git a/Engine/Plugins/Developer/Concert/ConcertUI/ConcertServerUI/Source/ConcertServerUI/Private/Widgets/SessionTabs/Live/PackageViewer/ConcertSessionPackageViewerController.h b/Engine/Plugins/Developer/Concert/ConcertUI/ConcertServerUI/Source/ConcertServerUI/Private/Widgets/SessionTabs/Live/PackageViewer/ConcertSessionPackageViewerController.h
new file mode 100644
--- /dev/null
+++ b/Engine/Plugins/Developer/Concert/ConcertUI/ConcertServerUI/Source/ConcertServerUI/Private/Widgets/SessionTabs/Live/PackageViewer/ConcertSessionPackageViewerController.h
@@ -0,0 +1,39 @@
+﻿// Copyright Epic Games, Inc. All Rights Reserved.
+
+#pragma once
+
+#include "CoreMinimal.h"
+
+struct FConcertSyncActivity;
+struct FConcertClientInfo;
+struct FConcertSyncPackageEventMetaData;
+struct FConcertSessionActivity;
+class IConcertSyncServer;
+class IConcertServerSession;
+class SConcertSessionPackageViewer;
+
+/** Manages SConcertSessionPackageViewer for a server. */
+class FConcertSessionPackageViewerController
+{
+public:
+
+	FConcertSessionPackageViewerController(TSharedRef<IConcertServerSession> InspectedSession, TSharedRef<IConcertSyncServer> SyncServer);
+	~FConcertSessionPackageViewerController();
+	
+	void ReloadActivities() const;
+	TSharedRef<SConcertSessionPackageViewer> GetPackageViewer() const { return PackageViewer; }
+
+private:
+	
+	TSharedRef<IConcertServerSession> InspectedSession;
+	TSharedRef<IConcertSyncServer> SyncServer;
+	
+	/** The widget being managed */
+	TSharedRef<SConcertSessionPackageViewer> PackageViewer;
+	
+	TSharedRef<SConcertSessionPackageViewer> MakePackageViewer() const;
+	
+	TOptional<FConcertClientInfo> GetClientInfo(FGuid ClientId) const;
+	bool GetPackageEvent(const FConcertSessionActivity& Activity, FConcertSyncPackageEventMetaData& OutPackageEvent) const;
+	void OnSessionProduced(const FConcertSyncActivity& ProducedActivity) const;
+};
diff --git a/Engine/Plugins/Developer/Concert/ConcertUI/ConcertServerUI/Source/ConcertServerUI/Private/Widgets/SessionTabs/Live/PackageViewer/SConcertSessionPackageViewer.cpp b/Engine/Plugins/Developer/Concert/ConcertUI/ConcertServerUI/Source/ConcertServerUI/Private/Widgets/SessionTabs/Live/PackageViewer/SConcertSessionPackageViewer.cpp
new file mode 100644
--- /dev/null
+++ b/Engine/Plugins/Developer/Concert/ConcertUI/ConcertServerUI/Source/ConcertServerUI/Private/Widgets/SessionTabs/Live/PackageViewer/SConcertSessionPackageViewer.cpp
@@ -0,0 +1,97 @@
+﻿// Copyright Epic Games, Inc. All Rights Reserved.
+
+#include "SConcertSessionPackageViewer.h"
+
+#include "Session/Activity/PredefinedActivityColumns.h"
+#include "Session/Activity/SConcertSessionActivities.h"
+#include "Widgets/Input/SSearchBox.h"
+
+#define LOCTEXT_NAMESPACE "UnrealMultiUserUI"
+
+void SConcertSessionPackageViewer::Construct(const FArguments& InArgs)
+{
+	using namespace UE::ConcertSharedSlate;
+	
+	ActivityListViewOptions = MakeShared<FConcertSessionActivitiesOptions>();
+	ActivityListViewOptions->bEnableConnectionActivityFiltering = false;
+	ActivityListViewOptions->bEnableLockActivityFiltering = false;
+	ActivityListViewOptions->bEnablePackageActivityFiltering = false;
+	ActivityListViewOptions->bEnableTransactionActivityFiltering = false;
+
+	SAssignNew(ActivityListView, SConcertSessionActivities)
+		.OnGetPackageEvent(InArgs._GetPackageEvent)
+		.OnFetchActivities(InArgs._FetchInitialActivities)
+		.OnMapActivityToClient(InArgs._GetClientInfo)
+		.HighlightText(this, &SConcertSessionPackageViewer::HighlightSearchedText)
+		.Columns({
+			// TODO: Create Icon (depending on EConcertPackageUpdateType), Size, and Version
+		})
+		.TimeFormat(ActivityListViewOptions.Get(), &FConcertSessionActivitiesOptions::GetTimeFormat)
+		.ConnectionActivitiesVisibility(EVisibility::Collapsed)
+		.LockActivitiesVisibility(EVisibility::Collapsed)
+		.PackageActivitiesVisibility(EVisibility::Visible)
+		.TransactionActivitiesVisibility(EVisibility::Collapsed)
+		.DetailsAreaVisibility(EVisibility::Collapsed)
+		.IsAutoScrollEnabled(true);
+
+	ChildSlot
+	[
+		SNew(SVerticalBox)
+		+SVerticalBox::Slot()
+		.AutoHeight()
+		.Padding(1, 1)
+		[
+			SAssignNew(SearchBox, SSearchBox)
+			.HintText(LOCTEXT("SearchHint", "Search..."))
+			.OnTextChanged(this, &SConcertSessionPackageViewer::OnSearchTextChanged)
+			.OnTextCommitted(this, &SConcertSessionPackageViewer::OnSearchTextCommitted)
+			.DelayChangeNotificationsWhileTyping(true)
+		]
+
+		+SVerticalBox::Slot()
+		[
+			ActivityListView.ToSharedRef()
+		]
+
+		+SVerticalBox::Slot()
+		.AutoHeight()
+		[
+			ActivityListViewOptions->MakeStatusBar(
+				TAttribute<int32>(ActivityListView.Get(), &SConcertSessionActivities::GetTotalActivityNum),
+				TAttribute<int32>(ActivityListView.Get(), &SConcertSessionActivities::GetDisplayedActivityNum)
+				)
+		]
+	];
+}
+
+void SConcertSessionPackageViewer::ResetActivityList()
+{
+	ActivityListView->ResetActivityList();
+}
+
+void SConcertSessionPackageViewer::AppendActivity(FConcertSessionActivity Activity)
+{
+	const TSharedRef<FConcertSessionActivity> NewActivity = MakeShared<FConcertSessionActivity>(MoveTemp(Activity));
+	ActivityListView->Append(NewActivity);
+}
+
+void SConcertSessionPackageViewer::OnSearchTextChanged(const FText& InSearchText)
+{
+	SearchedText = InSearchText;
+	SearchBox->SetError(ActivityListView->UpdateTextFilter(InSearchText));
+}
+
+void SConcertSessionPackageViewer::OnSearchTextCommitted(const FText& InFilterText, ETextCommit::Type CommitType)
+{
+	if (!InFilterText.EqualTo(SearchedText))
+	{
+		OnSearchTextChanged(InFilterText);
+	}
+}
+
+FText SConcertSessionPackageViewer::HighlightSearchedText() const
+{
+	return SearchedText;
+}
+
+#undef LOCTEXT_NAMESPACE
diff --git a/Engine/Plugins/Developer/Concert/ConcertUI/ConcertServerUI/Source/ConcertServerUI/Private/Widgets/SessionTabs/Live/PackageViewer/SConcertSessionPackageViewer.h b/Engine/Plugins/Developer/Concert/ConcertUI/ConcertServerUI/Source/ConcertServerUI/Private/Widgets/SessionTabs/Live/PackageViewer/SConcertSessionPackageViewer.h
new file mode 100644
--- /dev/null
+++ b/Engine/Plugins/Developer/Concert/ConcertUI/ConcertServerUI/Source/ConcertServerUI/Private/Widgets/SessionTabs/Live/PackageViewer/SConcertSessionPackageViewer.h
@@ -0,0 +1,53 @@
+﻿// Copyright Epic Games, Inc. All Rights Reserved.
+
+#pragma once
+
+#include "CoreMinimal.h"
+#include "Session/Activity/SConcertSessionActivities.h"
+#include "Widgets/SCompoundWidget.h"
+#include "Widgets/DeclarativeSyntaxSupport.h"
+
+class FConcertSessionActivitiesOptions;
+
+/** Displays package activities that happened in a concert session. */
+class SConcertSessionPackageViewer : public SCompoundWidget
+{
+public:
+
+	struct FConcertPackageActivity
+	{
+		FConcertSyncActivity Base;
+		FConcertSyncPackageEventData PackageEvent;
+	};
+	
+	SLATE_BEGIN_ARGS(SConcertSessionPackageViewer) {}
+		SLATE_EVENT(SConcertSessionActivities::FGetPackageEvent, GetPackageEvent)
+		SLATE_EVENT(SConcertSessionActivities::FGetActivityClientInfoFunc, GetClientInfo)
+		SLATE_EVENT(SConcertSessionActivities::FFetchActivitiesFunc, FetchInitialActivities)
+	SLATE_END_ARGS()
+
+	void Construct(const FArguments& InArgs);
+
+	void ResetActivityList();
+	void AppendActivity(FConcertSessionActivity Activity);
+
+private:
+
+	TSharedPtr<SConcertSessionActivities> ActivityListView;
+	
+	/** Controls the activity list view options */
+	TSharedPtr<FConcertSessionActivitiesOptions> ActivityListViewOptions;
+
+	/** The widget used to enter the text to search. */
+	TSharedPtr<SSearchBox> SearchBox;
+
+	/** The searched text to highlight. */
+	FText SearchedText;
+
+	/** Invoked when the text in the search box widget changes. */
+	void OnSearchTextChanged(const FText& InSearchText);
+	/** Invoked when the text in the search box widget is committed. */
+	void OnSearchTextCommitted(const FText& InFilterText, ETextCommit::Type CommitType);
+	/** Returns the text to highlight when the search bar has a text set. */
+	FText HighlightSearchedText() const;
+};
diff --git a/Engine/Plugins/Developer/Concert/ConcertUI/ConcertServerUI/Source/ConcertServerUI/Private/Widgets/SessionTabs/Live/SConcertSessionInspector.cpp b/Engine/Plugins/Developer/Concert/ConcertUI/ConcertServerUI/Source/ConcertServerUI/Private/Widgets/SessionTabs/Live/SConcertSessionInspector.cpp
new file mode 100644
--- /dev/null
+++ b/Engine/Plugins/Developer/Concert/ConcertUI/ConcertServerUI/Source/ConcertServerUI/Private/Widgets/SessionTabs/Live/SConcertSessionInspector.cpp
@@ -0,0 +1,127 @@
+﻿// Copyright Epic Games, Inc. All Rights Reserved.
+
+#include "SConcertSessionInspector.h"
+
+#include "Framework/Docking/TabManager.h"
+#include "Session/History/SSessionHistory.h"
+#include "Widgets/Docking/SDockTab.h"
+#include "Widgets/Layout/SBorder.h"
+#include "Widgets/SBoxPanel.h"
+
+#define LOCTEXT_NAMESPACE "SConcertSessionInspector"
+
+const FName SConcertSessionInspector::HistoryTabId("HistoryTabId");
+const FName SConcertSessionInspector::SessionContentTabId("SessionContentTabId");
+const FName SConcertSessionInspector::ConnectionMonitorTabId("ConnectionMonitorTabId");
+
+void SConcertSessionInspector::Construct(const FArguments& InArgs, const FRequiredArgs& RequiredArgs)
+{
+	ChildSlot
+	[
+		SNew(SBorder)
+		.BorderImage(FAppStyle::Get().GetBrush("ToolPanel.GroupBorder"))
+		.Padding(FMargin(1.0f, 2.0f))
+		[
+			SNew(SVerticalBox)
+
+			// Content
+			+SVerticalBox::Slot()
+			.FillHeight(1.f)
+			[
+				SNew(SBorder)
+				.BorderImage(FAppStyle::Get().GetBrush("Brushes.Title"))
+				.Padding(FMargin(0.f, 0.f, 0.f, 5.0f)) // Visually separate status bar from tabs
+				[
+					CreateTabs(InArgs, RequiredArgs)
+				]
+			]
+
+			// Status bar
+			+SVerticalBox::Slot()
+			.AutoHeight()
+			.HAlign(HAlign_Fill)
+			[
+				InArgs._StatusBar.Widget
+			]
+		]
+	];
+}
+
+TSharedRef<SWidget> SConcertSessionInspector::CreateTabs(const FArguments& InArgs, const FRequiredArgs& RequiredArgs)
+{
+	TabManager = FGlobalTabmanager::Get()->NewTabManager(RequiredArgs.ConstructUnderMajorTab);
+	
+	TabManager->RegisterTabSpawner(HistoryTabId, FOnSpawnTab::CreateSP(this, &SConcertSessionInspector::SpawnActivityHistory, RequiredArgs.SessionHistory))
+		.SetDisplayName(LOCTEXT("ActivityHistoryLabel", "History"));
+	
+	TabManager->RegisterTabSpawner(SessionContentTabId, FOnSpawnTab::CreateSP(this, &SConcertSessionInspector::SpawnSessionContent, RequiredArgs.PackageViewer))
+		.SetDisplayName(LOCTEXT("SessionContentLabel", "Session Content"));
+	
+	TabManager->RegisterTabSpawner(ConnectionMonitorTabId, FOnSpawnTab::CreateSP(this, &SConcertSessionInspector::SpawnConnectionMonitor))
+		.SetDisplayName(LOCTEXT("ConnectionMonitorLabel", "Connection Monitor"));
+
+	const TSharedRef<FTabManager::FLayout> Layout = FTabManager::NewLayout("ConcertSessionLayout_v0.3")
+		->AddArea
+		(
+			FTabManager::NewPrimaryArea()
+				->SetOrientation(Orient_Horizontal)
+				->Split
+				(
+					FTabManager::NewStack()
+					->SetSizeCoefficient(0.5f)
+					->AddTab(HistoryTabId, ETabState::OpenedTab)
+				)
+				->Split
+				(
+					FTabManager::NewSplitter()
+					->SetSizeCoefficient(0.5f)
+					->SetOrientation(Orient_Vertical)
+					->Split
+					(
+						FTabManager::NewStack()
+						->SetSizeCoefficient(0.5f)
+						->AddTab(SessionContentTabId, ETabState::OpenedTab)
+					)
+					->Split
+					(
+						FTabManager::NewStack()
+						->SetSizeCoefficient(0.5f)
+						->AddTab(ConnectionMonitorTabId, ETabState::OpenedTab)
+					)
+				)
+		);
+
+	return TabManager->RestoreFrom(Layout, RequiredArgs.ConstructUnderWindow).ToSharedRef();
+}
+
+TSharedRef<SDockTab> SConcertSessionInspector::SpawnActivityHistory(const FSpawnTabArgs& Args, TSharedRef<SSessionHistory> SessionHistory)
+{
+	return SNew(SDockTab)
+		.Label(LOCTEXT("ActivityHistoryLabel", "History"))
+		.TabRole(PanelTab)
+		[
+			SessionHistory
+		];
+}
+
+TSharedRef<SDockTab> SConcertSessionInspector::SpawnSessionContent(const FSpawnTabArgs& Args, TSharedRef<SConcertSessionPackageViewer> PackageViewer)
+{
+	return SNew(SDockTab)
+		.Label(LOCTEXT("SessionContentLabel", "Session Content"))
+		.TabRole(PanelTab)
+		[
+			PackageViewer
+		];
+}
+
+TSharedRef<SDockTab> SConcertSessionInspector::SpawnConnectionMonitor(const FSpawnTabArgs& Args)
+{
+	return SNew(SDockTab)
+		.Label(LOCTEXT("ConnectionMonitorLabel", "Connection Monitor"))
+		.TabRole(PanelTab)
+		[
+			SNullWidget::NullWidget
+		];
+}
+
+#undef LOCTEXT_NAMESPACE
diff --git a/Engine/Plugins/Developer/Concert/ConcertUI/ConcertServerUI/Source/ConcertServerUI/Private/Widgets/SessionTabs/Live/SConcertSessionInspector.h b/Engine/Plugins/Developer/Concert/ConcertUI/ConcertServerUI/Source/ConcertServerUI/Private/Widgets/SessionTabs/Live/SConcertSessionInspector.h
new file mode 100644
--- /dev/null
+++ b/Engine/Plugins/Developer/Concert/ConcertUI/ConcertServerUI/Source/ConcertServerUI/Private/Widgets/SessionTabs/Live/SConcertSessionInspector.h
@@ -0,0 +1,64 @@
+﻿// Copyright Epic Games, Inc. All Rights Reserved.
+
+#pragma once
+
+#include "CoreMinimal.h"
+#include "Session/History/SSessionHistory.h"
+#include "PackageViewer/SConcertSessionPackageViewer.h"
+#include "Widgets/SCompoundWidget.h"
+#include "Widgets/DeclarativeSyntaxSupport.h"
+
+class FTabManager;
+class FSpawnTabArgs;
+class SDockTab;
+class SWindow;
+
+/**
+ * Designed to be the content of a tab showing:
+ *  - activity history (transactions stored on the server as well as who made those transactions)
+ *  - session content (list of session data saved during a Multi-user session)
+ *  - connection monitor (details about the connected clients on the given session and network info)
+ * Implements view in model-view-controller pattern.
+ */
+class SConcertSessionInspector : public SCompoundWidget
+{
+public:
+
+	struct FRequiredArgs
+	{
+		TSharedRef<SDockTab> ConstructUnderMajorTab;
+		TSharedRef<SWindow> ConstructUnderWindow;
+		TSharedRef<SSessionHistory> SessionHistory;
+		TSharedRef<SConcertSessionPackageViewer> PackageViewer;
+
+		FRequiredArgs(TSharedRef<SDockTab> ConstructUnderMajorTab, TSharedRef<SWindow> ConstructUnderWindow,
+				TSharedRef<SSessionHistory> SessionHistoryController,
+				TSharedRef<SConcertSessionPackageViewer> PackageViewerController)
+			: ConstructUnderMajorTab(MoveTemp(ConstructUnderMajorTab))
+			, ConstructUnderWindow(MoveTemp(ConstructUnderWindow))
+			, SessionHistory(MoveTemp(SessionHistoryController))
+			, PackageViewer(MoveTemp(PackageViewerController))
+		{}
+	};
+
+	static const FName HistoryTabId;
+	static const FName SessionContentTabId;
+	static const FName ConnectionMonitorTabId;
+
+	SLATE_BEGIN_ARGS(SConcertSessionInspector) {}
+		SLATE_NAMED_SLOT(FArguments, StatusBar)
+	SLATE_END_ARGS()
+
+	void Construct(const FArguments& InArgs, const FRequiredArgs& RequiredArgs);
+
+private:
+	
+	/** Holds the tab manager that manages the front-end's tabs. */
+	TSharedPtr<FTabManager> TabManager;
+
+	// Spawn tabs
+	TSharedRef<SWidget> CreateTabs(const FArguments& InArgs, const FRequiredArgs& RequiredArgs);
+	TSharedRef<SDockTab> SpawnActivityHistory(const FSpawnTabArgs& Args, TSharedRef<SSessionHistory> SessionHistory);
+	TSharedRef<SDockTab> SpawnSessionContent(const FSpawnTabArgs& Args, TSharedRef<SConcertSessionPackageViewer> PackageViewer);
+	TSharedRef<SDockTab> SpawnConnectionMonitor(const FSpawnTabArgs& Args);
+};
diff --git a/Engine/Plugins/Developer/Concert/ConcertUI/ConcertServerUI/Source/ConcertServerUI/Private/Widgets/SessionTabs/Live/ServerSessionHistoryController.cpp b/Engine/Plugins/Developer/Concert/ConcertUI/ConcertServerUI/Source/ConcertServerUI/Private/Widgets/SessionTabs/Live/ServerSessionHistoryController.cpp
new file mode 100644
--- /dev/null
+++ b/Engine/Plugins/Developer/Concert/ConcertUI/ConcertServerUI/Source/ConcertServerUI/Private/Widgets/SessionTabs/Live/ServerSessionHistoryController.cpp
@@ -0,0 +1,92 @@
+﻿// Copyright Epic Games, Inc. All Rights Reserved.
+
+#include "ServerSessionHistoryController.h"
+
+#include "ConcertSyncSessionDatabase.h"
+#include "IConcertSession.h"
+#include "IConcertSyncServer.h"
+
+FServerSessionHistoryController::FServerSessionHistoryController(TSharedRef<IConcertServerSession> InspectedSession, TSharedRef<IConcertSyncServer> SyncServer)
+	: InspectedSession(MoveTemp(InspectedSession))
+	, SyncServer(MoveTemp(SyncServer))
+{
+	ReloadActivities();
+
+	if (const TOptional<FConcertSyncSessionDatabaseNonNullPtr> Database = SyncServer->GetLiveSessionDatabase(InspectedSession->GetId()))
+	{
+		Database->OnActivityProduced().AddRaw(this, &FServerSessionHistoryController::OnSessionProduced);
+	}
+}
+
+FServerSessionHistoryController::~FServerSessionHistoryController()
+{
+	if (const TOptional<FConcertSyncSessionDatabaseNonNullPtr> Database = SyncServer->GetLiveSessionDatabase(InspectedSession->GetId()))
+	{
+		Database->OnActivityProduced().RemoveAll(this);
+	}
+}
+
+void FServerSessionHistoryController::GetActivities(int64 MaximumNumberOfActivities, TMap<FGuid, FConcertClientInfo>& OutEndpointClientInfoMap, TArray<FConcertSessionActivity>& OutFetchedActivities) const
+{
+	if (const TOptional<FConcertSyncSessionDatabaseNonNullPtr> Database = SyncServer->GetLiveSessionDatabase(InspectedSession->GetId()))
+	{
+		OutEndpointClientInfoMap.Reset();
+		OutFetchedActivities.Reset();
+
+		int64 LastActivityId = INDEX_NONE;
+		Database->GetActivityMaxId(LastActivityId);
+		const int64 FirstActivityIdToFetch = FMath::Max<int64>(1, LastActivityId - MaximumNumberOfActivities);
+		Database->EnumerateActivitiesInRange(FirstActivityIdToFetch, MaximumNumberOfActivities, [this, &Database, &OutEndpointClientInfoMap, &OutFetchedActivities](FConcertSyncActivity&& InActivity)
+		{
+			if (!OutEndpointClientInfoMap.Contains(InActivity.EndpointId))
+			{
+				FConcertSyncEndpointData EndpointData;
+				if (Database->GetEndpoint(InActivity.EndpointId, EndpointData))
+				{
+					OutEndpointClientInfoMap.Add(InActivity.EndpointId, EndpointData.ClientInfo);
+				}
+			}
+
+			FStructOnScope ActivitySummary;
+			if (InActivity.EventSummary.GetPayload(ActivitySummary))
+			{
+				OutFetchedActivities.Emplace(MoveTemp(InActivity), MoveTemp(ActivitySummary));
+			}
+
+			return EBreakBehavior::Continue;
+		});
+	}
+}
+
+bool FServerSessionHistoryController::GetPackageEvent(const FConcertSessionActivity& Activity, FConcertSyncPackageEventMetaData& OutPackageEvent) const
+{
+	if (const TOptional<FConcertSyncSessionDatabaseNonNullPtr> Database = SyncServer->GetLiveSessionDatabase(InspectedSession->GetId()))
+	{
+		return Database->GetPackageEventMetaData(Activity.Activity.EventId, OutPackageEvent.PackageRevision, OutPackageEvent.PackageInfo);
+	}
+	
+	return false;
+}
+
+TFuture<TOptional<FConcertSyncTransactionEvent>> FServerSessionHistoryController::GetTransactionEvent(const FConcertSessionActivity& Activity) const
+{
+	if (const TOptional<FConcertSyncSessionDatabaseNonNullPtr> Database = SyncServer->GetLiveSessionDatabase(InspectedSession->GetId()))
+	{
+		return FindOrRequestTransactionEvent(*Database, Activity.Activity.EventId);
+	}
+	
+	return MakeFulfilledPromise<TOptional<FConcertSyncTransactionEvent>>().GetFuture(); // Not found.
+}
+
+TFuture<TOptional<FConcertSyncTransactionEvent>> FServerSessionHistoryController::FindOrRequestTransactionEvent(const FConcertSyncSessionDatabase& Database, const int64 TransactionEventId) const
+{
+	FConcertSyncTransactionEvent TransactionEvent;
+	return Database.GetTransactionEvent(TransactionEventId, TransactionEvent, false)
+		?	MakeFulfilledPromise<TOptional<FConcertSyncTransactionEvent>>(MoveTemp(TransactionEvent)).GetFuture()
+		:	MakeFulfilledPromise<TOptional<FConcertSyncTransactionEvent>>().GetFuture();
+}
+
+void FServerSessionHistoryController::OnSessionProduced(const FConcertSyncActivity& ProducedActivity)
+{
+	ReloadActivities();
+}
diff --git a/Engine/Plugins/Developer/Concert/ConcertUI/ConcertServerUI/Source/ConcertServerUI/Private/Widgets/SessionTabs/Live/ServerSessionHistoryController.h b/Engine/Plugins/Developer/Concert/ConcertUI/ConcertServerUI/Source/ConcertServerUI/Private/Widgets/SessionTabs/Live/ServerSessionHistoryController.h
new file mode 100644
--- /dev/null
+++ b/Engine/Plugins/Developer/Concert/ConcertUI/ConcertServerUI/Source/ConcertServerUI/Private/Widgets/SessionTabs/Live/ServerSessionHistoryController.h
@@ -0,0 +1,36 @@
+﻿// Copyright Epic Games, Inc. All Rights Reserved.
+
+#pragma once
+
+#include "CoreMinimal.h"
+#include "Session/History/AbstractSessionHistoryController.h"
+
+struct FConcertSyncActivity;
+class FConcertSyncSessionDatabase;
+class IConcertSyncServer;
+class IConcertServerSession;
+
+class FServerSessionHistoryController : public FAbstractSessionHistoryController
+{
+public:
+
+	FServerSessionHistoryController(TSharedRef<IConcertServerSession> InspectedSession, TSharedRef<IConcertSyncServer> SyncServer);
+	virtual ~FServerSessionHistoryController() override;
+	
+protected:
+
+	//~ Begin FAbstractSessionHistoryController Interface
+	virtual void GetActivities(int64 MaximumNumberOfActivities, TMap<FGuid, FConcertClientInfo>& OutEndpointClientInfoMap, TArray<FConcertSessionActivity>& OutFetchedActivities) const override;
+	virtual bool GetPackageEvent(const FConcertSessionActivity& Activity, FConcertSyncPackageEventMetaData& OutPackageEvent) const override;
+	virtual TFuture<TOptional<FConcertSyncTransactionEvent>> GetTransactionEvent(const FConcertSessionActivity& Activity) const override;
+	//~ End FAbstractSessionHistoryController Interface
+
+private:
+
+	TSharedRef<IConcertServerSession> InspectedSession;
+	TSharedRef<IConcertSyncServer> SyncServer;
+
+	TFuture<TOptional<FConcertSyncTransactionEvent>> FindOrRequestTransactionEvent(const FConcertSyncSessionDatabase& Database, const int64 TransactionEventId) const;
+	
+	void OnSessionProduced(const FConcertSyncActivity& ProducedActivity);
+};
diff --git a/Engine/Plugins/Developer/Concert/ConcertUI/ConcertServerUI/Source/ConcertServerUI/Private/Widgets/Tabs/AbstractConcertSessionTab.cpp b/Engine/Plugins/Developer/Concert/ConcertUI/ConcertServerUI/Source/ConcertServerUI/Private/Widgets/Tabs/AbstractConcertSessionTab.cpp
deleted file mode 100644
diff --git a/Engine/Plugins/Developer/Concert/ConcertUI/ConcertServerUI/Source/ConcertServerUI/Private/Widgets/Tabs/AbstractConcertSessionTab.h b/Engine/Plugins/Developer/Concert/ConcertUI/ConcertServerUI/Source/ConcertServerUI/Private/Widgets/Tabs/AbstractConcertSessionTab.h
deleted file mode 100644
diff --git a/Engine/Plugins/Developer/Concert/ConcertUI/ConcertServerUI/Source/ConcertServerUI/Private/Widgets/Tabs/Archived/ArchivedConcertSessionTab.cpp b/Engine/Plugins/Developer/Concert/ConcertUI/ConcertServerUI/Source/ConcertServerUI/Private/Widgets/Tabs/Archived/ArchivedConcertSessionTab.cpp
deleted file mode 100644
diff --git a/Engine/Plugins/Developer/Concert/ConcertUI/ConcertServerUI/Source/ConcertServerUI/Private/Widgets/Tabs/Archived/ArchivedConcertSessionTab.h b/Engine/Plugins/Developer/Concert/ConcertUI/ConcertServerUI/Source/ConcertServerUI/Private/Widgets/Tabs/Archived/ArchivedConcertSessionTab.h
deleted file mode 100644
diff --git a/Engine/Plugins/Developer/Concert/ConcertUI/ConcertServerUI/Source/ConcertServerUI/Private/Widgets/Tabs/Archived/SConcertArchivedSessionInspector.cpp b/Engine/Plugins/Developer/Concert/ConcertUI/ConcertServerUI/Source/ConcertServerUI/Private/Widgets/Tabs/Archived/SConcertArchivedSessionInspector.cpp
deleted file mode 100644
diff --git a/Engine/Plugins/Developer/Concert/ConcertUI/ConcertServerUI/Source/ConcertServerUI/Private/Widgets/Tabs/Archived/SConcertArchivedSessionInspector.h b/Engine/Plugins/Developer/Concert/ConcertUI/ConcertServerUI/Source/ConcertServerUI/Private/Widgets/Tabs/Archived/SConcertArchivedSessionInspector.h
deleted file mode 100644
